Enable Linux support for EXR

diff --git a/Engine/Plugins/Media/ImgMedia/Source/ImgMedia/ImgMedia.Build.cs b/Engine/Plugins/Media/ImgMedia/Source/ImgMedia/ImgMedia.Build.cs
--- a/Engine/Plugins/Media/ImgMedia/Source/ImgMedia/ImgMedia.Build.cs
+++ b/Engine/Plugins/Media/ImgMedia/Source/ImgMedia/ImgMedia.Build.cs
@@ -45,6 +45,7 @@ namespace UnrealBuildTool.Rules
 				});
 
 			if ((Target.Platform == UnrealTargetPlatform.Mac) ||
+				(Target.Platform == UnrealTargetPlatform.Linux) ||
 				(Target.Platform == UnrealTargetPlatform.Win32) ||
 				(Target.Platform == UnrealTargetPlatform.Win64))
 			{
diff --git a/Engine/Plugins/Media/ImgMedia/Source/ImgMedia/Private/ImgMediaPrivate.h b/Engine/Plugins/Media/ImgMedia/Source/ImgMedia/Private/ImgMediaPrivate.h
--- a/Engine/Plugins/Media/ImgMedia/Source/ImgMedia/Private/ImgMediaPrivate.h
+++ b/Engine/Plugins/Media/ImgMedia/Source/ImgMedia/Private/ImgMediaPrivate.h
@@ -10,8 +10,8 @@
 class FQueuedThreadPool;
 
 
-/** The OpenEXR image reader is supported on macOS and Windows only. */
-#define IMGMEDIA_EXR_SUPPORTED_PLATFORM (PLATFORM_MAC || PLATFORM_WINDOWS)
+/** The OpenEXR image reader is supported on macOS, Windows, and Unix only. */
+#define IMGMEDIA_EXR_SUPPORTED_PLATFORM (PLATFORM_MAC || PLATFORM_WINDOWS || PLATFORM_UNIX)
 
 /** Whether to use a separate thread pool for image frame deallocations. */
 #define USE_IMGMEDIA_DEALLOC_POOL UE_BUILD_DEBUG
diff --git a/Engine/Plugins/Media/ImgMedia/Source/OpenExrWrapper/OpenExrWrapper.Build.cs b/Engine/Plugins/Media/ImgMedia/Source/OpenExrWrapper/OpenExrWrapper.Build.cs
--- a/Engine/Plugins/Media/ImgMedia/Source/OpenExrWrapper/OpenExrWrapper.Build.cs
+++ b/Engine/Plugins/Media/ImgMedia/Source/OpenExrWrapper/OpenExrWrapper.Build.cs
@@ -19,6 +19,7 @@ namespace UnrealBuildTool.Rules
 
             if ((Target.Platform == UnrealTargetPlatform.Win64) ||
                 (Target.Platform == UnrealTargetPlatform.Win32) ||
+                (Target.Platform == UnrealTargetPlatform.Linux) ||
                 (Target.Platform == UnrealTargetPlatform.Mac))
             {
                 AddEngineThirdPartyPrivateStaticDependencies(Target, "UEOpenExr");
