Only enable for Linux x86

diff --git a/Engine/Plugins/Media/ImgMedia/Source/ImgMedia/ImgMedia.Build.cs b/Engine/Plugins/Media/ImgMedia/Source/ImgMedia/ImgMedia.Build.cs
--- a/Engine/Plugins/Media/ImgMedia/Source/ImgMedia/ImgMedia.Build.cs
+++ b/Engine/Plugins/Media/ImgMedia/Source/ImgMedia/ImgMedia.Build.cs
@@ -44,10 +44,12 @@ namespace UnrealBuildTool.Rules
 					"TimeManagement",
 				});
 
+			bool bLinuxEnabled = Target.Platform == UnrealTargetPlatform.Linux && Target.Architecture.StartsWith("x86_64");
+
 			if ((Target.Platform == UnrealTargetPlatform.Mac) ||
-				(Target.Platform == UnrealTargetPlatform.Linux) ||
 				(Target.Platform == UnrealTargetPlatform.Win32) ||
-				(Target.Platform == UnrealTargetPlatform.Win64))
+				(Target.Platform == UnrealTargetPlatform.Win64) ||
+				bLinuxEnabled)
 			{
 				PrivateDependencyModuleNames.Add("OpenExrWrapper");
 			}
diff --git a/Engine/Plugins/Media/ImgMedia/Source/ImgMedia/Private/ImgMediaPrivate.h b/Engine/Plugins/Media/ImgMedia/Source/ImgMedia/Private/ImgMediaPrivate.h
--- a/Engine/Plugins/Media/ImgMedia/Source/ImgMedia/Private/ImgMediaPrivate.h
+++ b/Engine/Plugins/Media/ImgMedia/Source/ImgMedia/Private/ImgMediaPrivate.h
@@ -11,7 +11,7 @@ class FQueuedThreadPool;
 
 
 /** The OpenEXR image reader is supported on macOS, Windows, and Linux only. */
-#define IMGMEDIA_EXR_SUPPORTED_PLATFORM (PLATFORM_MAC || PLATFORM_WINDOWS || PLATFORM_LINUX)
+#define IMGMEDIA_EXR_SUPPORTED_PLATFORM (PLATFORM_MAC || PLATFORM_WINDOWS || (PLATFORM_LINUX && PLATFORM_CPU_X86_FAMILY))
 
 /** Whether to use a separate thread pool for image frame deallocations. */
 #define USE_IMGMEDIA_DEALLOC_POOL UE_BUILD_DEBUG
diff --git a/Engine/Plugins/Media/ImgMedia/Source/OpenExrWrapper/OpenExrWrapper.Build.cs b/Engine/Plugins/Media/ImgMedia/Source/OpenExrWrapper/OpenExrWrapper.Build.cs
--- a/Engine/Plugins/Media/ImgMedia/Source/OpenExrWrapper/OpenExrWrapper.Build.cs
+++ b/Engine/Plugins/Media/ImgMedia/Source/OpenExrWrapper/OpenExrWrapper.Build.cs
@@ -19,10 +19,12 @@ namespace UnrealBuildTool.Rules
 					"TimeManagement",
 				});
 
+			bool bLinuxEnabled = Target.Platform == UnrealTargetPlatform.Linux && Target.Architecture.StartsWith("x86_64");
+
             if ((Target.Platform == UnrealTargetPlatform.Win64) ||
                 (Target.Platform == UnrealTargetPlatform.Win32) ||
-                (Target.Platform == UnrealTargetPlatform.Linux) ||
-                (Target.Platform == UnrealTargetPlatform.Mac))
+                (Target.Platform == UnrealTargetPlatform.Mac) ||
+				bLinuxEnabled)
             {
                 AddEngineThirdPartyPrivateStaticDependencies(Target, "UEOpenExr");
                 AddEngineThirdPartyPrivateStaticDependencies(Target, "zlib");
