ImgMedia plugin no longer creates threads on startup

diff --git a/Engine/Plugins/Media/ImgMedia/Source/ImgMedia/Private/ImgMediaModule.cpp b/Engine/Plugins/Media/ImgMedia/Source/ImgMedia/Private/ImgMediaModule.cpp
--- a/Engine/Plugins/Media/ImgMedia/Source/ImgMedia/Private/ImgMediaModule.cpp
+++ b/Engine/Plugins/Media/ImgMedia/Source/ImgMedia/Private/ImgMediaModule.cpp
@@ -16,8 +16,70 @@ DEFINE_LOG_CATEGORY(LogImgMedia);
 
 
 #if USE_IMGMEDIA_DEALLOC_POOL
-	FQueuedThreadPool* GImgMediaThreadPoolSlow = nullptr;
-#endif
+struct FImgMediaThreadPool
+{
+public:
+
+	FImgMediaThreadPool() :
+		Pool(nullptr),
+		bHasInit(false)
+	{
+	}
+
+	~FImgMediaThreadPool()
+	{
+		Reset();
+	}
+
+	void Reset()
+	{
+		FScopeLock Lock(&CriticalSection);
+		if (Pool != nullptr)
+		{
+			Pool->Destroy();
+			Pool = nullptr;
+		}
+
+		bHasInit = false;
+	}
+
+	FQueuedThreadPool* GetThreadPool()
+	{
+		FScopeLock Lock(&CriticalSection);
+		if (bHasInit)
+		{
+			return Pool;
+		}
+
+		// initialize worker thread pools
+		if (FPlatformProcess::SupportsMultithreading())
+		{
+			// initialize dealloc thread pool
+			const int32 ThreadPoolSize = 1;
+			const uint32 StackSize = 128 * 1024;
+
+			Pool = FQueuedThreadPool::Allocate();
+			verify(Pool->Create(ThreadPoolSize, StackSize, TPri_Normal));
+		}
+
+		bHasInit = true;
+
+		return Pool;
+	}
+
+private:
+	FCriticalSection CriticalSection;
+	FQueuedThreadPool* Pool;
+	bool bHasInit;
+};
+
+FImgMediaThreadPool ImgMediaThreadPool;
+
+FQueuedThreadPool* GetImgMediaThreadPoolSlow()
+{
+	return ImgMediaThreadPool.GetThreadPool();
+}
+#endif // USE_IMGMEDIA_DEALLOC_POOL
 
 
 /**
@@ -37,6 +99,11 @@ public:
 
 	virtual TSharedPtr<IMediaPlayer, ESPMode::ThreadSafe> CreatePlayer(IMediaEventSink& EventSink) override
 	{
+		if (!Scheduler.IsValid())
+		{
+			InitScheduler();
+		}
+
 		return MakeShared<FImgMediaPlayer, ESPMode::ThreadSafe>(EventSink, Scheduler.ToSharedRef());
 	}
 
@@ -46,29 +113,7 @@ public:
 
 	virtual void StartupModule() override
 	{
-		// initialize scheduler
-		Scheduler = MakeShared<FImgMediaScheduler, ESPMode::ThreadSafe>();
-		Scheduler->Initialize();
-
-		IMediaModule* MediaModule = FModuleManager::LoadModulePtr<IMediaModule>("Media");
-
-		if (MediaModule != nullptr)
-		{
-			MediaModule->GetClock().AddSink(Scheduler.ToSharedRef());
-		}
-
-		// initialize worker thread pools
-		if (FPlatformProcess::SupportsMultithreading())
-		{
-#if USE_IMGMEDIA_DEALLOC_POOL
-			// initialize dealloc thread pool
-			const int32 ThreadPoolSize = FPlatformProperties::IsServerOnly() ? 1 : FPlatformMisc::NumberOfWorkerThreadsToSpawn();
-			const uint32 StackSize = 128 * 1024;
 
-			GImgMediaThreadPoolSlow = FQueuedThreadPool::Allocate();
-			verify(GImgMediaThreadPoolSlow->Create(ThreadPoolSize, StackSize, TPri_Normal));
-#endif
-		}
 	}
 
 	virtual void ShutdownModule() override
@@ -76,16 +121,26 @@ public:
 		Scheduler.Reset();
 
 #if USE_IMGMEDIA_DEALLOC_POOL
-		if (GImgMediaThreadPoolSlow != nullptr)
-		{
-			GImgMediaThreadPoolSlow->Destroy();
-			GImgMediaThreadPoolSlow = nullptr;
-		}
+		ImgMediaThreadPool.Reset();
 #endif
 	}
 
 private:
 
+	void InitScheduler()
+	{
+		// initialize scheduler
+		Scheduler = MakeShared<FImgMediaScheduler, ESPMode::ThreadSafe>();
+		Scheduler->Initialize();
+
+		IMediaModule* MediaModule = FModuleManager::LoadModulePtr<IMediaModule>("Media");
+
+		if (MediaModule != nullptr)
+		{
+			MediaModule->GetClock().AddSink(Scheduler.ToSharedRef());
+		}
+	}
+
 	TSharedPtr<FImgMediaScheduler, ESPMode::ThreadSafe> Scheduler;
 };
 
diff --git a/Engine/Plugins/Media/ImgMedia/Source/ImgMedia/Private/ImgMediaPrivate.h b/Engine/Plugins/Media/ImgMedia/Source/ImgMedia/Private/ImgMediaPrivate.h
--- a/Engine/Plugins/Media/ImgMedia/Source/ImgMedia/Private/ImgMediaPrivate.h
+++ b/Engine/Plugins/Media/ImgMedia/Source/ImgMedia/Private/ImgMediaPrivate.h
@@ -23,7 +23,7 @@ DECLARE_LOG_CATEGORY_EXTERN(LogImgMedia, Log, All);
 
 #if USE_IMGMEDIA_DEALLOC_POOL
 	/** Thread pool used for deleting image frame buffers. */
-	extern FQueuedThreadPool* GImgMediaThreadPoolSlow;
+	extern FQueuedThreadPool* GetImgMediaThreadPoolSlow();
 #endif
 
 
diff --git a/Engine/Plugins/Media/ImgMedia/Source/ImgMedia/Private/Readers/ExrImgMediaReader.cpp b/Engine/Plugins/Media/ImgMedia/Source/ImgMedia/Private/Readers/ExrImgMediaReader.cpp
--- a/Engine/Plugins/Media/ImgMedia/Source/ImgMedia/Private/Readers/ExrImgMediaReader.cpp
+++ b/Engine/Plugins/Media/ImgMedia/Source/ImgMedia/Private/Readers/ExrImgMediaReader.cpp
@@ -57,23 +57,22 @@ bool FExrImgMediaReader::ReadFrame(const FString& ImagePath, FImgMediaFrame& Out
 
 	auto BufferDeleter = [BufferSize](void* ObjectToDelete) {
 #if USE_IMGMEDIA_DEALLOC_POOL
-		// free buffers on the thread pool, because memory allocators may perform
-		// expensive operations, such as filling the memory with debug values
-		TFunction<void()> FreeBufferTask = [ObjectToDelete, BufferSize]()
+		if (FQueuedThreadPool* ImgMediaThreadPoolSlow = GetImgMediaThreadPoolSlow())
 		{
-#endif
-			FMemory::Free(ObjectToDelete);
-#if USE_IMGMEDIA_DEALLOC_POOL
-		};
-
-		if (GImgMediaThreadPoolSlow != nullptr)
-		{
-			AsyncPool(*GImgMediaThreadPoolSlow, FreeBufferTask);
+			// free buffers on the thread pool, because memory allocators may perform
+			// expensive operations, such as filling the memory with debug values
+			TFunction<void()> FreeBufferTask = [ObjectToDelete]()
+			{
+				FMemory::Free(ObjectToDelete);
+			};
+			AsyncPool(*ImgMediaThreadPoolSlow, FreeBufferTask);
 		}
 		else
 		{
-			FreeBufferTask();
+			FMemory::Free(ObjectToDelete);
 		}
+#else
+		FMemory::Free(ObjectToDelete);
 #endif
 	};
 
