MediaFrameworks: Added global cache for image sequences.

diff --git a/Engine/Plugins/Media/ImgMedia/Source/ImgMedia/ImgMedia.Build.cs b/Engine/Plugins/Media/ImgMedia/Source/ImgMedia/ImgMedia.Build.cs
--- a/Engine/Plugins/Media/ImgMedia/Source/ImgMedia/ImgMedia.Build.cs
+++ b/Engine/Plugins/Media/ImgMedia/Source/ImgMedia/ImgMedia.Build.cs
@@ -32,6 +32,7 @@ namespace UnrealBuildTool.Rules
 				new string[] {
 					"ImgMedia/Private",
 					"ImgMedia/Private/Assets",
+					"ImgMedia/Private/GlobalCache",
 					"ImgMedia/Private/Loader",
 					"ImgMedia/Private/Player",
 					"ImgMedia/Private/Readers",
diff --git a/Engine/Plugins/Media/ImgMedia/Source/ImgMedia/Private/GlobalCache/ImgMediaGlobalCache.cpp b/Engine/Plugins/Media/ImgMedia/Source/ImgMedia/Private/GlobalCache/ImgMediaGlobalCache.cpp
new file mode 100644
--- /dev/null
+++ b/Engine/Plugins/Media/ImgMedia/Source/ImgMedia/Private/GlobalCache/ImgMediaGlobalCache.cpp
@@ -0,0 +1,223 @@
+// Copyright 1998-2019 Epic Games, Inc. All Rights Reserved.
+
+#include "ImgMediaGlobalCache.h"
+#include "IImgMediaReader.h"
+#include "ImgMediaPrivate.h"
+#include "ImgMediaSettings.h"
+
+
+FImgMediaGlobalCache::FImgMediaGlobalCache()
+	: LeastRecent(nullptr)
+	, MostRecent(nullptr)
+	, CurrentSize(0)
+	, MaxSize(0)
+{
+}
+
+FImgMediaGlobalCache::~FImgMediaGlobalCache()
+{
+	Shutdown();
+}
+
+void FImgMediaGlobalCache::Initialize()
+{
+	auto Settings = GetDefault<UImgMediaSettings>();
+	MaxSize = Settings->GlobalCacheSizeGB * 1024 * 1024 * 1024;
+}
+
+void FImgMediaGlobalCache::Shutdown()
+{
+	FScopeLock Lock(&CriticalSection);
+	Empty();
+}
+
+void FImgMediaGlobalCache::AddFrame(const FName& Sequence, int32 Index, const TSharedPtr<FImgMediaFrame, ESPMode::ThreadSafe>& Frame)
+{
+	FScopeLock Lock(&CriticalSection);
+
+	// Make sure we have enough space in the cache to add this new frame.
+	SIZE_T FrameSize = Frame->Info.UncompressedSize;
+	if (FrameSize <= MaxSize)
+	{
+		// Empty cache until we have enough space.
+		while (CurrentSize + FrameSize > MaxSize)
+		{
+			FName* RemoveSequencePtr = MapLeastRecentToSequence.Find(LeastRecent);
+			FName RemoveSequence = RemoveSequencePtr != nullptr ? *RemoveSequencePtr : FName();
+			Remove(RemoveSequence, *LeastRecent);
+		}
+
+		// Create new entry.
+		FImgMediaGlobalCacheEntry* NewEntry = new FImgMediaGlobalCacheEntry(Index, Frame);
+		MapFrameToEntry.Emplace(TPair<FName, int32>(Sequence, Index), NewEntry);
+	
+		MarkAsRecent(Sequence, *NewEntry);
+
+		CurrentSize += FrameSize;
+	}
+	else
+	{
+		UE_LOG(LogImgMedia, Warning, TEXT("Global cache size %d is smaller than frame size %d."), MaxSize, FrameSize);
+	}
+}
+
+TSharedPtr<FImgMediaFrame, ESPMode::ThreadSafe>* FImgMediaGlobalCache::FindAndTouch(const FName& Sequence, int32 Index)
+{
+	FScopeLock Lock(&CriticalSection);
+
+	TSharedPtr<FImgMediaFrame, ESPMode::ThreadSafe>* Frame = nullptr;
+
+	FImgMediaGlobalCacheEntry** Entry = MapFrameToEntry.Find(TPair<FName, int32>(Sequence, Index));
+
+	if ((Entry != nullptr) && (*Entry != nullptr))
+	{
+		Frame = &((*Entry)->Frame);
+
+		// Mark this as the most recent.
+		Unlink(Sequence, **Entry);
+		MarkAsRecent(Sequence, **Entry);
+	}
+
+	return Frame;
+}
+
+void FImgMediaGlobalCache::GetIndices(const FName& Sequence, TArray<int32>& OutIndices) const
+{
+	FScopeLock Lock(&CriticalSection);
+
+	// Get most recent entry in this sequence.
+	FImgMediaGlobalCacheEntry* const* CurrentPtr = MapSequenceToMostRecentEntry.Find(Sequence);
+	const FImgMediaGlobalCacheEntry* Current = CurrentPtr != nullptr ? *CurrentPtr : nullptr;
+
+	// Loop over all entries in the sequence.
+	while (Current != nullptr)
+	{
+		OutIndices.Add(Current->Index);
+		Current = Current->LessRecentSequence;
+	}
+}
+
+void FImgMediaGlobalCache::Remove(const FName& Sequence, FImgMediaGlobalCacheEntry& Entry)
+{
+	// Remove from cache.
+	Unlink(Sequence, Entry);
+
+	// Update current cache size.
+	SIZE_T FrameSize = Entry.Frame->Info.UncompressedSize;
+	CurrentSize -= FrameSize;
+
+	// Delete entry.
+	MapFrameToEntry.Remove(TPair<FName, int32>(Sequence, Entry.Index));
+	delete &Entry;
+}
+
+void FImgMediaGlobalCache::MarkAsRecent(const FName& Sequence, FImgMediaGlobalCacheEntry& Entry)
+{
+	// Mark most recent.
+	Entry.LessRecent = MostRecent;
+	if (MostRecent != nullptr)
+	{
+		MostRecent->MoreRecent = &Entry;
+	}
+	MostRecent = &Entry;
+
+	// Mark most recent in sequence.
+	FImgMediaGlobalCacheEntry** SequenceMostRecentPtr = MapSequenceToMostRecentEntry.Find(Sequence);
+	FImgMediaGlobalCacheEntry* SequenceMostRecent = (SequenceMostRecentPtr != nullptr) ? (*SequenceMostRecentPtr) : nullptr;
+	Entry.LessRecentSequence = SequenceMostRecent;
+	if (SequenceMostRecent != nullptr)
+	{
+		SequenceMostRecent->MoreRecentSequence = &Entry;
+	}
+	else
+	{
+		// If we did not have a most recent one, then this is the first in the sequence.
+		MapLeastRecentToSequence.Emplace(&Entry, Sequence);
+	}
+	MapSequenceToMostRecentEntry.Emplace(Sequence, &Entry);
+
+	// If LeastRecent is null, then set it now.
+	if (LeastRecent == nullptr)
+	{
+		LeastRecent = &Entry;
+	}
+}
+
+void FImgMediaGlobalCache::Unlink(const FName& Sequence, FImgMediaGlobalCacheEntry& Entry)
+{
+	// Remove from link.
+	if (Entry.LessRecent != nullptr)
+	{
+		Entry.LessRecent->MoreRecent = Entry.MoreRecent;
+	}
+	else if (LeastRecent == &Entry)
+	{
+		LeastRecent = Entry.MoreRecent;
+	}
+
+	if (Entry.MoreRecent != nullptr)
+	{
+		Entry.MoreRecent->LessRecent = Entry.LessRecent;
+	}
+	else if (MostRecent == &Entry)
+	{
+		MostRecent = Entry.LessRecent;
+	}
+
+	Entry.LessRecent = nullptr;
+	Entry.MoreRecent = nullptr;
+
+	// Remove from sequence link.
+	if (Entry.LessRecentSequence != nullptr)
+	{
+		Entry.LessRecentSequence->MoreRecentSequence = Entry.MoreRecentSequence;
+	}
+	else
+	{
+		MapLeastRecentToSequence.Remove(&Entry);
+		if (Entry.MoreRecentSequence != nullptr)
+		{
+			MapLeastRecentToSequence.Emplace(Entry.MoreRecentSequence, Sequence);
+		}
+	}
+
+	if (Entry.MoreRecentSequence != nullptr)
+	{
+		Entry.MoreRecentSequence->LessRecentSequence = Entry.LessRecentSequence;
+	}
+	else
+	{
+		// Update most recent in sequence.
+		FImgMediaGlobalCacheEntry** MostRecentSequence = MapSequenceToMostRecentEntry.Find(Sequence);
+		if (MostRecentSequence != nullptr)
+		{
+			if (*MostRecentSequence == &Entry)
+			{
+				MapSequenceToMostRecentEntry.Emplace(Sequence, Entry.LessRecentSequence);
+			}
+		}
+	}
+
+	Entry.LessRecent = nullptr;
+	Entry.LessRecentSequence = nullptr;
+	Entry.MoreRecent = nullptr;
+	Entry.MoreRecentSequence = nullptr;
+}
+
+void FImgMediaGlobalCache::Empty()
+{
+	while (LeastRecent != nullptr)
+	{
+		FImgMediaGlobalCacheEntry *Entry = LeastRecent;
+		LeastRecent = Entry->MoreRecent;
+
+		delete Entry;
+	}
+
+	MostRecent = nullptr;
+	CurrentSize = 0;
+
+	MapSequenceToMostRecentEntry.Empty();
+	MapLeastRecentToSequence.Empty();
+	MapFrameToEntry.Empty();
+}
diff --git a/Engine/Plugins/Media/ImgMedia/Source/ImgMedia/Private/GlobalCache/ImgMediaGlobalCache.h b/Engine/Plugins/Media/ImgMedia/Source/ImgMedia/Private/GlobalCache/ImgMediaGlobalCache.h
new file mode 100644
--- /dev/null
+++ b/Engine/Plugins/Media/ImgMedia/Source/ImgMedia/Private/GlobalCache/ImgMediaGlobalCache.h
@@ -0,0 +1,143 @@
+// Copyright 1998-2019 Epic Games, Inc. All Rights Reserved.
+
+#pragma once
+
+#include "CoreTypes.h"
+
+struct FImgMediaFrame;
+
+/**
+ * A global cache for all ImgMedia players.
+ *
+ * Uses Least Recently Used (LRU).
+ */
+class FImgMediaGlobalCache
+{
+public:
+	/**
+	 * Constructor.
+	 */
+	FImgMediaGlobalCache();
+
+	/**
+	 * Desctructor.
+	  */
+	~FImgMediaGlobalCache();
+
+	/**
+	 * Initialize the cache.
+	 */
+	void Initialize();
+
+	/**
+	 * Shut down the cache.
+	 */
+	void Shutdown();
+
+	/**
+	 * Adds a frame to the cache.
+	 *
+	 * @param	Sequence		Indentifying name of this sequence.
+	 * @param	Index			Index of frame to add.
+	 * @param	Frame			Actual frame to add.
+	 */
+	void AddFrame(const FName& Sequence, int32 Index, const TSharedPtr<FImgMediaFrame, ESPMode::ThreadSafe>& Frame);
+
+	/**
+	 * Find the entry with the specified sequence and index and mark it as the most recently used.
+	 *
+	 * @param	Sequence		Indentifying name of this sequence.
+	 * @param	Index			Index of frame to touch.
+	 */
+	TSharedPtr<FImgMediaFrame, ESPMode::ThreadSafe>* FindAndTouch(const FName& Sequence, int32 Index);
+
+	/**
+	 * Find the indices of all cached entries of a particular sequence.
+	 *
+	 * @param	Sequence		Indentifying name of this sequence.
+	 * @param OutIndices Will contain the collection of indices.
+	 */
+	void GetIndices(const FName& Sequence, TArray<int32>& OutIndices) const;
+
+private:
+
+	/** An entry in the cache. */
+	struct FImgMediaGlobalCacheEntry
+	{
+		/** Frame index. */
+		int32 Index;
+		/** Actual frame. */
+		TSharedPtr<FImgMediaFrame, ESPMode::ThreadSafe> Frame;
+
+		/** Previous entry in the cache. */
+		FImgMediaGlobalCacheEntry* LessRecent;
+		/** Previous entry in the cache that is part of the same sequence. */
+		FImgMediaGlobalCacheEntry* LessRecentSequence;
+		/** Next entry in the cache. */
+		FImgMediaGlobalCacheEntry* MoreRecent;
+		/** Next entry in the cache that is part of the same sequence. */
+		FImgMediaGlobalCacheEntry* MoreRecentSequence;
+
+		/** Constructor. */
+		FImgMediaGlobalCacheEntry(int32 InIndex, const TSharedPtr<FImgMediaFrame, ESPMode::ThreadSafe>& InFrame)
+			: Index(InIndex)
+			, Frame(InFrame)
+			, LessRecent(nullptr)
+			, LessRecentSequence(nullptr)
+			, MoreRecent(nullptr)
+			, MoreRecentSequence(nullptr)
+		{
+		}
+	};
+
+	/** Entry that was used first. */
+	FImgMediaGlobalCacheEntry* LeastRecent;
+	/** Entry that was used last. */
+	FImgMediaGlobalCacheEntry* MostRecent;
+
+	/** Maps a sequence name to the most recent cache entry of that sequence. */
+	TMap<FName, FImgMediaGlobalCacheEntry*> MapSequenceToMostRecentEntry;
+	/** Maps a cache entry that is the least recent entry to the name of its sequence. */
+	TMap<FImgMediaGlobalCacheEntry*, FName> MapLeastRecentToSequence;
+	/** Maps a sequence name and frame index to an entry in the cache. */
+	TMap<TPair<FName, int32>, FImgMediaGlobalCacheEntry*> MapFrameToEntry;
+
+	/** Current size of the cache in bytes. */
+	SIZE_T CurrentSize;
+	/** Maximum size of the cache in bytes. */
+	SIZE_T MaxSize;
+
+	/** Critical section for synchronizing access to Frames. */
+	mutable FCriticalSection CriticalSection;
+	
+	/**
+	 * Removes an entry from the cache and deletes the entry.
+	 *
+	 * @param	Sequence		Indentifying name of this sequence.
+	 * @param	Entry			Entry to remove.
+	 */
+	void Remove(const FName& Sequence, FImgMediaGlobalCacheEntry& Entry);
+
+	/**
+	 * Inserts an entry into the cache as the most recent entry.
+	 * Assumes the entry is not already in the cache.
+	 *
+	 * @param	Sequence		Indentifying name of this sequence.
+	 * @param	Entry			Entry to mark as recent.
+	 */
+	void MarkAsRecent(const FName& Sequence, FImgMediaGlobalCacheEntry& Entry);
+
+	/**
+	 * Removes an entry from the cache but does not delete the entry.
+	 *
+	 * @param	Sequence		Indentifying name of this sequence.
+	 * @param	Entry			Entry to remove.
+	 */
+	void Unlink(const FName& Sequence, FImgMediaGlobalCacheEntry& Entry);
+
+	/**
+	 * Empties the cache.
+	 */
+	void Empty();
+};
+
diff --git a/Engine/Plugins/Media/ImgMedia/Source/ImgMedia/Private/ImgMediaModule.cpp b/Engine/Plugins/Media/ImgMedia/Source/ImgMedia/Private/ImgMediaModule.cpp
--- a/Engine/Plugins/Media/ImgMedia/Source/ImgMedia/Private/ImgMediaModule.cpp
+++ b/Engine/Plugins/Media/ImgMedia/Source/ImgMedia/Private/ImgMediaModule.cpp
@@ -7,6 +7,7 @@
 #include "Misc/QueuedThreadPool.h"
 #include "Modules/ModuleManager.h"
 
+#include "ImgMediaGlobalCache.h"
 #include "ImgMediaPlayer.h"
 #include "ImgMediaScheduler.h"
 #include "IImgMediaModule.h"
@@ -103,8 +104,12 @@ public:
 		{
 			InitScheduler();
 		}
+		if (!GlobalCache.IsValid())
+		{
+			InitGlobalCache();
+		}
 
-		return MakeShared<FImgMediaPlayer, ESPMode::ThreadSafe>(EventSink, Scheduler.ToSharedRef());
+		return MakeShared<FImgMediaPlayer, ESPMode::ThreadSafe>(EventSink, Scheduler.ToSharedRef(), GlobalCache.ToSharedRef());
 	}
 
 public:
@@ -119,6 +124,7 @@ public:
 	virtual void ShutdownModule() override
 	{
 		Scheduler.Reset();
+		GlobalCache.Reset();
 
 #if USE_IMGMEDIA_DEALLOC_POOL
 		ImgMediaThreadPool.Reset();
@@ -141,7 +147,15 @@ private:
 		}
 	}
 
+	void InitGlobalCache()
+	{
+		// Initialize global cache.
+		GlobalCache = MakeShared<FImgMediaGlobalCache, ESPMode::ThreadSafe>();
+		GlobalCache->Initialize();
+	}
+
 	TSharedPtr<FImgMediaScheduler, ESPMode::ThreadSafe> Scheduler;
+	TSharedPtr<FImgMediaGlobalCache, ESPMode::ThreadSafe> GlobalCache;
 };
 
 
diff --git a/Engine/Plugins/Media/ImgMedia/Source/ImgMedia/Private/Loader/ImgMediaLoader.cpp b/Engine/Plugins/Media/ImgMedia/Source/ImgMedia/Private/Loader/ImgMediaLoader.cpp
--- a/Engine/Plugins/Media/ImgMedia/Source/ImgMedia/Private/Loader/ImgMediaLoader.cpp
+++ b/Engine/Plugins/Media/ImgMedia/Source/ImgMedia/Private/Loader/ImgMediaLoader.cpp
@@ -1,6 +1,7 @@
 // Copyright 1998-2019 Epic Games, Inc. All Rights Reserved.
 
 #include "ImgMediaLoader.h"
+#include "ImgMediaGlobalCache.h"
 #include "ImgMediaPrivate.h"
 
 #include "Algo/Reverse.h"
@@ -35,17 +36,20 @@ DECLARE_CYCLE_STAT(TEXT("ImgMedia Loader Release Cache"), STAT_ImgMedia_LoaderRe
 /* FImgMediaLoader structors
  *****************************************************************************/
 
-FImgMediaLoader::FImgMediaLoader(const TSharedRef<FImgMediaScheduler, ESPMode::ThreadSafe>& InScheduler)
+FImgMediaLoader::FImgMediaLoader(const TSharedRef<FImgMediaScheduler, ESPMode::ThreadSafe>& InScheduler,
+	const TSharedRef<FImgMediaGlobalCache, ESPMode::ThreadSafe>& InGlobalCache)
 	: Frames(1)
 	, ImageWrapperModule(FModuleManager::LoadModuleChecked<IImageWrapperModule>("ImageWrapper"))
 	, Initialized(false)
 	, NumLoadAhead(0)
 	, NumLoadBehind(0)
 	, Scheduler(InScheduler)
+	, GlobalCache(InGlobalCache)
 	, SequenceDim(FIntPoint::ZeroValue)
 	, SequenceDuration(FTimespan::Zero())
 	, SequenceFrameRate(0, 0)
 	, LastRequestedFrame(INDEX_NONE)
+	, UseGlobalCache(false)
 {
 	UE_LOG(LogImgMedia, Verbose, TEXT("Loader %p: Created"), this);
 }
@@ -95,7 +99,14 @@ void FImgMediaLoader::GetCompletedTimeRanges(TRangeSet<FTimespan>& OutRangeSet)
 	FScopeLock Lock(&CriticalSection);
 
 	TArray<int32> CompletedFrames;
-	Frames.GetKeys(CompletedFrames);
+	if (UseGlobalCache)
+	{
+		GlobalCache->GetIndices(SequenceName, CompletedFrames);
+	}
+	else
+	{
+		Frames.GetKeys(CompletedFrames);
+	}
 	FrameNumbersToTimeRanges(CompletedFrames, OutRangeSet);
 }
 
@@ -111,7 +122,16 @@ TSharedPtr<FImgMediaTextureSample, ESPMode::ThreadSafe> FImgMediaLoader::GetFram
 
 	FScopeLock ScopeLock(&CriticalSection);
 
-	const TSharedPtr<FImgMediaFrame, ESPMode::ThreadSafe>* Frame = Frames.FindAndTouch(FrameIndex);
+	const TSharedPtr<FImgMediaFrame, ESPMode::ThreadSafe>* Frame;
+	if (UseGlobalCache)
+	{
+		Frame = GlobalCache->FindAndTouch(SequenceName, FrameIndex);
+	}
+	else
+	{
+		Frame = Frames.FindAndTouch(FrameIndex);
+	}
+
 
 	if (Frame == nullptr)
 	{
@@ -307,6 +327,7 @@ void FImgMediaLoader::LoadSequence(const FString& SequencePath, const FFrameRate
 
 	// initialize loader
 	auto Settings = GetDefault<UImgMediaSettings>();
+	UseGlobalCache = Settings->UseGlobalCache;
 
 	const FPlatformMemoryStats Stats = FPlatformMemory::GetStats();
 	const SIZE_T DesiredCacheSize = Settings->CacheSizeGB * 1024 * 1024 * 1024;
@@ -320,6 +341,9 @@ void FImgMediaLoader::LoadSequence(const FString& SequencePath, const FFrameRate
 	NumLoadAhead = NumFramesToLoad - NumLoadBehind;
 
 	Frames.Empty(NumFramesToLoad);
+
+	SequenceName = FName(*SequencePath);
+
 	Update(0, 0.0f, Loop);
 
 	// update info
@@ -453,7 +477,17 @@ void FImgMediaLoader::Update(int32 PlayHeadFrame, float PlayRate, bool Loop)
 
 	for (int32 FrameNumber : FramesToLoad)
 	{
-		if ((Frames.FindAndTouch(FrameNumber) == nullptr) && !QueuedFrameNumbers.Contains(FrameNumber))
+		bool NeedFrame = false;
+		if (UseGlobalCache)
+		{
+			NeedFrame = GlobalCache->FindAndTouch(SequenceName, FrameNumber) == nullptr;
+		}
+		else
+		{
+			NeedFrame = Frames.FindAndTouch(FrameNumber) == nullptr;
+		}
+		
+		if ((NeedFrame) && !QueuedFrameNumbers.Contains(FrameNumber))
 		{
 			PendingFrameNumbers.Add(FrameNumber);
 		}
@@ -476,7 +510,14 @@ void FImgMediaLoader::NotifyWorkComplete(FImgMediaLoaderWork& CompletedWork, int
 		if (Frame.IsValid())
 		{
 			UE_LOG(LogImgMedia, VeryVerbose, TEXT("Loader %p: Loaded frame %i"), this, FrameNumber);
-			Frames.Add(FrameNumber, Frame);
+			if (UseGlobalCache)
+			{
+				GlobalCache->AddFrame(SequenceName, FrameNumber, Frame);
+			}
+			else
+			{
+				Frames.Add(FrameNumber, Frame);
+			}
 		}
 	}
 
diff --git a/Engine/Plugins/Media/ImgMedia/Source/ImgMedia/Private/Loader/ImgMediaLoader.h b/Engine/Plugins/Media/ImgMedia/Source/ImgMedia/Private/Loader/ImgMediaLoader.h
--- a/Engine/Plugins/Media/ImgMedia/Source/ImgMedia/Private/Loader/ImgMediaLoader.h
+++ b/Engine/Plugins/Media/ImgMedia/Source/ImgMedia/Private/Loader/ImgMediaLoader.h
@@ -8,6 +8,7 @@
 #include "Misc/FrameRate.h"
 #include "Templates/SharedPointer.h"
 
+class FImgMediaGlobalCache;
 class FImgMediaLoaderWork;
 class FImgMediaScheduler;
 class FImgMediaTextureSample;
@@ -31,7 +32,8 @@ public:
 	 *
 	 * @param InScheduler The scheduler for image loading.
 	 */
-	FImgMediaLoader(const TSharedRef<FImgMediaScheduler, ESPMode::ThreadSafe>& InScheduler);
+	FImgMediaLoader(const TSharedRef<FImgMediaScheduler, ESPMode::ThreadSafe>& InScheduler,
+		const TSharedRef<FImgMediaGlobalCache, ESPMode::ThreadSafe>& InGlobalCache);
 
 	/** Virtual destructor. */
 	virtual ~FImgMediaLoader();
@@ -270,6 +272,9 @@ private:
 	/** The scheduler for image loading. */
 	TSharedPtr<FImgMediaScheduler, ESPMode::ThreadSafe> Scheduler;
 
+	/** The scheduler for image loading. */
+	TSharedPtr<FImgMediaGlobalCache, ESPMode::ThreadSafe> GlobalCache;
+
 	/** Width and height of the image sequence (in pixels) .*/
 	FIntPoint SequenceDim;
 
@@ -279,6 +284,9 @@ private:
 	/** Frame rate of the currently loaded sequence. */
 	FFrameRate SequenceFrameRate;
 
+	/** Identifying name of sequence files. */
+	FName SequenceName;
+
 private:
 
 	/** Index of the previously requested frame. */
@@ -292,4 +300,7 @@ private:
 
 	/** Object pool for reusable work items. */
 	TArray<FImgMediaLoaderWork*> WorkPool;
+
+	/** True if we are using the global cache, false to use the local cache. */
+	bool UseGlobalCache;
 };
diff --git a/Engine/Plugins/Media/ImgMedia/Source/ImgMedia/Private/Player/ImgMediaPlayer.cpp b/Engine/Plugins/Media/ImgMedia/Source/ImgMedia/Private/Player/ImgMediaPlayer.cpp
--- a/Engine/Plugins/Media/ImgMedia/Source/ImgMedia/Private/Player/ImgMediaPlayer.cpp
+++ b/Engine/Plugins/Media/ImgMedia/Source/ImgMedia/Private/Player/ImgMediaPlayer.cpp
@@ -32,7 +32,8 @@ const FTimespan HackDeltaTimeOffset(1);
 /* FImgMediaPlayer structors
  *****************************************************************************/
 
-FImgMediaPlayer::FImgMediaPlayer(IMediaEventSink& InEventSink, const TSharedRef<FImgMediaScheduler, ESPMode::ThreadSafe>& InScheduler)
+FImgMediaPlayer::FImgMediaPlayer(IMediaEventSink& InEventSink, const TSharedRef<FImgMediaScheduler, ESPMode::ThreadSafe>& InScheduler,
+	const TSharedRef<FImgMediaGlobalCache, ESPMode::ThreadSafe>& InGlobalCache)
 	: CurrentDuration(FTimespan::Zero())
 	, CurrentRate(0.0f)
 	, CurrentState(EMediaState::Closed)
@@ -44,6 +45,7 @@ FImgMediaPlayer::FImgMediaPlayer(IMediaEventSink& InEventSink, const TSharedRef<
 	, Scheduler(InScheduler)
 	, SelectedVideoTrack(INDEX_NONE)
 	, ShouldLoop(false)
+	, GlobalCache(InGlobalCache)
 { }
 
 
@@ -179,7 +181,7 @@ bool FImgMediaPlayer::Open(const FString& Url, const IMediaOptions* Options)
 	}
 
 	// initialize image loader on a separate thread
-	Loader = MakeShared<FImgMediaLoader, ESPMode::ThreadSafe>(Scheduler.ToSharedRef());
+	Loader = MakeShared<FImgMediaLoader, ESPMode::ThreadSafe>(Scheduler.ToSharedRef(), GlobalCache.ToSharedRef());
 	Scheduler->RegisterLoader(Loader.ToSharedRef());
 
 	const FString SequencePath = Url.RightChop(6);
diff --git a/Engine/Plugins/Media/ImgMedia/Source/ImgMedia/Private/Player/ImgMediaPlayer.h b/Engine/Plugins/Media/ImgMedia/Source/ImgMedia/Private/Player/ImgMediaPlayer.h
--- a/Engine/Plugins/Media/ImgMedia/Source/ImgMedia/Private/Player/ImgMediaPlayer.h
+++ b/Engine/Plugins/Media/ImgMedia/Source/ImgMedia/Private/Player/ImgMediaPlayer.h
@@ -15,6 +15,7 @@ class FImgMediaScheduler;
 class IImgMediaReader;
 class IMediaEventSink;
 class IMediaTextureSample;
+class FImgMediaGlobalCache;
 
 
 /**
@@ -36,7 +37,8 @@ public:
 	 * @param InEventSink The object that receives media events from this player.
 	 * @param InScheduler The image loading scheduler to use.
 	 */
-	FImgMediaPlayer(IMediaEventSink& InEventSink, const TSharedRef<FImgMediaScheduler, ESPMode::ThreadSafe>& InScheduler);
+	FImgMediaPlayer(IMediaEventSink& InEventSink, const TSharedRef<FImgMediaScheduler, ESPMode::ThreadSafe>& InScheduler,
+		const TSharedRef<FImgMediaGlobalCache, ESPMode::ThreadSafe>& InGlobalCache);
 
 	/** Virtual destructor. */
 	virtual ~FImgMediaPlayer();
@@ -153,4 +155,7 @@ private:
 
 	/** Should the video loop to the beginning at completion */
     bool ShouldLoop;
+
+	/** The global cache to use. */
+	TSharedPtr<FImgMediaGlobalCache, ESPMode::ThreadSafe> GlobalCache;
 };
diff --git a/Engine/Plugins/Media/ImgMedia/Source/ImgMediaFactory/Private/ImgMediaSettings.cpp b/Engine/Plugins/Media/ImgMedia/Source/ImgMediaFactory/Private/ImgMediaSettings.cpp
--- a/Engine/Plugins/Media/ImgMedia/Source/ImgMediaFactory/Private/ImgMediaSettings.cpp
+++ b/Engine/Plugins/Media/ImgMedia/Source/ImgMediaFactory/Private/ImgMediaSettings.cpp
@@ -12,6 +12,8 @@ UImgMediaSettings::UImgMediaSettings()
 	, CacheSizeGB(1.0f)
 	, CacheThreads(8)
 	, CacheThreadStackSizeKB(128)
+	, GlobalCacheSizeGB(1.0f)
+	, UseGlobalCache(true)
 	, ExrDecoderThreads(0)
 	, DefaultProxy(TEXT("proxy"))
 	, UseDefaultProxy(false)
diff --git a/Engine/Plugins/Media/ImgMedia/Source/ImgMediaFactory/Public/ImgMediaSettings.h b/Engine/Plugins/Media/ImgMedia/Source/ImgMediaFactory/Public/ImgMediaSettings.h
--- a/Engine/Plugins/Media/ImgMedia/Source/ImgMediaFactory/Public/ImgMediaSettings.h
+++ b/Engine/Plugins/Media/ImgMedia/Source/ImgMediaFactory/Public/ImgMediaSettings.h
@@ -44,6 +44,14 @@ public:
 	UPROPERTY(config, EditAnywhere, Category=Caching, meta=(ClampMin=128), AdvancedDisplay)
 	int32 CacheThreadStackSizeKB;
 
+	/** Maximum size of the global look-ahead cache (in GB; default = 1 GB). */
+	UPROPERTY(config, EditAnywhere, Category = Caching, meta = (ClampMin = 0))
+	float GlobalCacheSizeGB;
+
+	/** Whether to use the global cache or not. */
+	UPROPERTY(config, EditAnywhere, Category = Caching)
+	bool UseGlobalCache;
+
 public:
 
 	/** Number of worker threads to use when decoding EXR images (0 = auto). */
