MediaFrameworks: Added missing #includes.

diff --git a/Engine/Plugins/Media/ImgMedia/Source/ImgMedia/Private/GlobalCache/ImgMediaGlobalCache.h b/Engine/Plugins/Media/ImgMedia/Source/ImgMedia/Private/GlobalCache/ImgMediaGlobalCache.h
--- a/Engine/Plugins/Media/ImgMedia/Source/ImgMedia/Private/GlobalCache/ImgMediaGlobalCache.h
+++ b/Engine/Plugins/Media/ImgMedia/Source/ImgMedia/Private/GlobalCache/ImgMediaGlobalCache.h
@@ -3,6 +3,8 @@
 #pragma once
 
 #include "CoreTypes.h"
+#include "Containers/Map.h"
+#include "Templates/SharedPointer.h"
 
 struct FImgMediaFrame;
 
