MediaFramework: Fixed ImgMedia not using the global cache when getting info for the first frame.

diff --git a/Engine/Plugins/Media/ImgMedia/Source/ImgMedia/Private/Loader/ImgMediaLoader.cpp b/Engine/Plugins/Media/ImgMedia/Source/ImgMedia/Private/Loader/ImgMediaLoader.cpp
--- a/Engine/Plugins/Media/ImgMedia/Source/ImgMedia/Private/Loader/ImgMediaLoader.cpp
+++ b/Engine/Plugins/Media/ImgMedia/Source/ImgMedia/Private/Loader/ImgMediaLoader.cpp
@@ -291,15 +291,30 @@ void FImgMediaLoader::LoadSequence(const FString& SequencePath, const FFrameRate
 		Reader = MakeShareable(new FGenericImgMediaReader(ImageWrapperModule));
 	}
 
+	const UImgMediaSettings* Settings = GetDefault<UImgMediaSettings>();
+	UseGlobalCache = Settings->UseGlobalCache;
+	SequenceName = FName(*SequencePath);
+
 	// fetch sequence attributes from first image
 	FImgMediaFrameInfo FirstFrameInfo;
-
-	if (!Reader->GetFrameInfo(ImagePaths[0], FirstFrameInfo))
 	{
-		UE_LOG(LogImgMedia, Error, TEXT("Failed to get frame information from first image in %s"), *SequencePath);
-		return;
-	}
+		// Try and get frame from the global cache.
+		const TSharedPtr<FImgMediaFrame, ESPMode::ThreadSafe>* Frame = nullptr;
+		if (UseGlobalCache)
+		{
+			Frame = GlobalCache->FindAndTouch(SequenceName, 0);
+		}
 
+		if (Frame)
+		{
+			FirstFrameInfo = Frame->Get()->Info;
+		}
+		else if (!Reader->GetFrameInfo(ImagePaths[0], FirstFrameInfo))
+		{
+			UE_LOG(LogImgMedia, Error, TEXT("Failed to get frame information from first image in %s"), *SequencePath);
+			return;
+		}
+	}
 	if (FirstFrameInfo.UncompressedSize == 0)
 	{
 		UE_LOG(LogImgMedia, Error, TEXT("The first image in sequence %s does not have a valid frame size"), *SequencePath);
@@ -326,9 +341,6 @@ void FImgMediaLoader::LoadSequence(const FString& SequencePath, const FFrameRate
 	SequenceDuration = FrameNumberToTime(ImagePaths.Num());
 
 	// initialize loader
-	auto Settings = GetDefault<UImgMediaSettings>();
-	UseGlobalCache = Settings->UseGlobalCache;
-
 	const FPlatformMemoryStats Stats = FPlatformMemory::GetStats();
 	const SIZE_T DesiredCacheSize = Settings->CacheSizeGB * 1024 * 1024 * 1024;
 	const SIZE_T CacheSize = FMath::Clamp(DesiredCacheSize, (SIZE_T)0, (SIZE_T)Stats.AvailablePhysical);
@@ -342,8 +354,6 @@ void FImgMediaLoader::LoadSequence(const FString& SequencePath, const FFrameRate
 
 	Frames.Empty(NumFramesToLoad);
 
-	SequenceName = FName(*SequencePath);
-
 	Update(0, 0.0f, Loop);
 
 	// update info
