MediaFramework: Fix for MediaPlayer deadlocking when using EXR.

diff --git a/Engine/Plugins/Media/ImgMedia/Source/ImgMedia/Private/Player/ImgMediaPlayer.cpp b/Engine/Plugins/Media/ImgMedia/Source/ImgMedia/Private/Player/ImgMediaPlayer.cpp
--- a/Engine/Plugins/Media/ImgMedia/Source/ImgMedia/Private/Player/ImgMediaPlayer.cpp
+++ b/Engine/Plugins/Media/ImgMedia/Source/ImgMedia/Private/Player/ImgMediaPlayer.cpp
@@ -46,6 +46,7 @@ FImgMediaPlayer::FImgMediaPlayer(IMediaEventSink& InEventSink, const TSharedRef<
 	, SelectedVideoTrack(INDEX_NONE)
 	, ShouldLoop(false)
 	, GlobalCache(InGlobalCache)
+	, RequestFrameHasRun(true)
 { }
 
 
@@ -306,6 +307,31 @@ void FImgMediaPlayer::TickInput(FTimespan DeltaTime, FTimespan /*Timecode*/)
 	{
 		Loader->RequestFrame(CurrentTime, CurrentRate, ShouldLoop);
 	}
+	RequestFrameHasRun = true;
+}
+
+void FImgMediaPlayer::ProcessVideoSamples()
+{
+	// Did we already run this frame?
+	if (RequestFrameHasRun)
+	{
+		RequestFrameHasRun = false;
+	}
+	else
+	{
+		// We are blocked... run stuff here as it will not get run normally.
+		if (Loader.IsValid())
+		{
+			if (SelectedVideoTrack == 0)
+			{
+				Loader->RequestFrame(CurrentTime, CurrentRate, ShouldLoop);
+			}
+		}
+		if (Scheduler.IsValid())
+		{
+			Scheduler->TickFetch(FTimespan::Zero(), FTimespan::Zero());
+		}
+	}
 }
 
 
diff --git a/Engine/Plugins/Media/ImgMedia/Source/ImgMedia/Private/Player/ImgMediaPlayer.h b/Engine/Plugins/Media/ImgMedia/Source/ImgMedia/Private/Player/ImgMediaPlayer.h
--- a/Engine/Plugins/Media/ImgMedia/Source/ImgMedia/Private/Player/ImgMediaPlayer.h
+++ b/Engine/Plugins/Media/ImgMedia/Source/ImgMedia/Private/Player/ImgMediaPlayer.h
@@ -60,6 +60,7 @@ public:
 	virtual bool Open(const FString& Url, const IMediaOptions* Options) override;
 	virtual bool Open(const TSharedRef<FArchive, ESPMode::ThreadSafe>& Archive, const FString& OriginalUrl, const IMediaOptions* Options) override;
 	virtual void TickInput(FTimespan DeltaTime, FTimespan Timecode) override;
+	virtual void ProcessVideoSamples() override;
 
 protected:
 
@@ -158,4 +159,7 @@ private:
 
 	/** The global cache to use. */
 	TSharedPtr<FImgMediaGlobalCache, ESPMode::ThreadSafe> GlobalCache;
+
+	/** True if we have run RequestFrame already for this frame. */
+	bool RequestFrameHasRun;
 };
diff --git a/Engine/Source/Runtime/Media/Public/IMediaPlayer.h b/Engine/Source/Runtime/Media/Public/IMediaPlayer.h
--- a/Engine/Source/Runtime/Media/Public/IMediaPlayer.h
+++ b/Engine/Source/Runtime/Media/Public/IMediaPlayer.h
@@ -312,6 +312,15 @@ public:
 		return true;
 	}
 
+	/**
+	 * Any extra processing that the player should do when FMediaPlayerFacade::ProcessVideoSamples
+	 * is run should be put here.
+	 */
+	virtual void ProcessVideoSamples()
+	{
+		// Override in child class if needed.
+	}
+
 public:
 
 	/** Virtual destructor. */
diff --git a/Engine/Source/Runtime/MediaUtils/Private/MediaPlayerFacade.cpp b/Engine/Source/Runtime/MediaUtils/Private/MediaPlayerFacade.cpp
--- a/Engine/Source/Runtime/MediaUtils/Private/MediaPlayerFacade.cpp
+++ b/Engine/Source/Runtime/MediaUtils/Private/MediaPlayerFacade.cpp
@@ -1176,6 +1176,12 @@ void FMediaPlayerFacade::ProcessSubtitleSamples(IMediaSamples& Samples, TRange<F
 
 void FMediaPlayerFacade::ProcessVideoSamples(IMediaSamples& Samples, TRange<FTimespan> TimeRange)
 {
+	// Let the player do some processing if needed.
+	if (Player.IsValid())
+	{
+		Player->ProcessVideoSamples();
+	}
+
 	TSharedPtr<IMediaTextureSample, ESPMode::ThreadSafe> Sample;
 
 	while (Samples.FetchVideo(TimeRange, Sample))
