MediaFramework: Fixed editor needing restart to update EXR settings.

diff --git a/Engine/Plugins/Media/ImgMedia/Source/ImgMedia/Private/GlobalCache/ImgMediaGlobalCache.cpp b/Engine/Plugins/Media/ImgMedia/Source/ImgMedia/Private/GlobalCache/ImgMediaGlobalCache.cpp
--- a/Engine/Plugins/Media/ImgMedia/Source/ImgMedia/Private/GlobalCache/ImgMediaGlobalCache.cpp
+++ b/Engine/Plugins/Media/ImgMedia/Source/ImgMedia/Private/GlobalCache/ImgMediaGlobalCache.cpp
@@ -22,13 +22,20 @@ FImgMediaGlobalCache::~FImgMediaGlobalCache()
 void FImgMediaGlobalCache::Initialize()
 {
 	auto Settings = GetDefault<UImgMediaSettings>();
-	MaxSize = Settings->GlobalCacheSizeGB * 1024 * 1024 * 1024;
+#if WITH_EDITOR
+	UpdateSettingsDelegateHandle = UImgMediaSettings::OnSettingsChanged().AddThreadSafeSP(this, &FImgMediaGlobalCache::UpdateSettings);
+#endif
+	UpdateSettings(Settings);
 }
 
 void FImgMediaGlobalCache::Shutdown()
 {
 	FScopeLock Lock(&CriticalSection);
 	Empty();
+
+#if WITH_EDITOR
+	UImgMediaSettings::OnSettingsChanged().Remove(UpdateSettingsDelegateHandle);
+#endif
 }
 
 void FImgMediaGlobalCache::AddFrame(const FName& Sequence, int32 Index, const TSharedPtr<FImgMediaFrame, ESPMode::ThreadSafe>& Frame)
@@ -40,12 +47,7 @@ void FImgMediaGlobalCache::AddFrame(const FName& Sequence, int32 Index, const TS
 	if (FrameSize <= MaxSize)
 	{
 		// Empty cache until we have enough space.
-		while (CurrentSize + FrameSize > MaxSize)
-		{
-			FName* RemoveSequencePtr = MapLeastRecentToSequence.Find(LeastRecent);
-			FName RemoveSequence = RemoveSequencePtr != nullptr ? *RemoveSequencePtr : FName();
-			Remove(RemoveSequence, *LeastRecent);
-		}
+		EnforceMaxSize(FrameSize);
 
 		// Create new entry.
 		FImgMediaGlobalCacheEntry* NewEntry = new FImgMediaGlobalCacheEntry(Index, Frame);
@@ -97,6 +99,16 @@ void FImgMediaGlobalCache::GetIndices(const FName& Sequence, TArray<int32>& OutI
 	}
 }
 
+void FImgMediaGlobalCache::EnforceMaxSize(SIZE_T Extra)
+{
+	while (CurrentSize + Extra > MaxSize)
+	{
+		FName* RemoveSequencePtr = MapLeastRecentToSequence.Find(LeastRecent);
+		FName RemoveSequence = RemoveSequencePtr != nullptr ? *RemoveSequencePtr : FName();
+		Remove(RemoveSequence, *LeastRecent);
+	}
+}
+
 void FImgMediaGlobalCache::Remove(const FName& Sequence, FImgMediaGlobalCacheEntry& Entry)
 {
 	// Remove from cache.
@@ -221,3 +233,9 @@ void FImgMediaGlobalCache::Empty()
 	MapLeastRecentToSequence.Empty();
 	MapFrameToEntry.Empty();
 }
+
+void FImgMediaGlobalCache::UpdateSettings(const UImgMediaSettings* Settings)
+{
+	MaxSize = Settings->GlobalCacheSizeGB * 1024 * 1024 * 1024;
+	EnforceMaxSize(0);
+}
diff --git a/Engine/Plugins/Media/ImgMedia/Source/ImgMedia/Private/GlobalCache/ImgMediaGlobalCache.h b/Engine/Plugins/Media/ImgMedia/Source/ImgMedia/Private/GlobalCache/ImgMediaGlobalCache.h
--- a/Engine/Plugins/Media/ImgMedia/Source/ImgMedia/Private/GlobalCache/ImgMediaGlobalCache.h
+++ b/Engine/Plugins/Media/ImgMedia/Source/ImgMedia/Private/GlobalCache/ImgMediaGlobalCache.h
@@ -7,13 +7,14 @@
 #include "Templates/SharedPointer.h"
 
 struct FImgMediaFrame;
+class UImgMediaSettings;
 
 /**
  * A global cache for all ImgMedia players.
  *
  * Uses Least Recently Used (LRU).
  */
-class FImgMediaGlobalCache
+class FImgMediaGlobalCache : public TSharedFromThis<FImgMediaGlobalCache, ESPMode::ThreadSafe>
 {
 public:
 	/**
@@ -111,7 +112,19 @@ private:
 
 	/** Critical section for synchronizing access to Frames. */
 	mutable FCriticalSection CriticalSection;
+
+#if WITH_EDITORONLY_DATA
+
+	/** Handle to registered UpdateSettings delegate. */
+	FDelegateHandle UpdateSettingsDelegateHandle;
+
+#endif
 	
+	/**
+	 * Empties the cache until the current size + Extra <= the max size of the cache.
+	 */
+	void EnforceMaxSize(SIZE_T Extra);
+
 	/**
 	 * Removes an entry from the cache and deletes the entry.
 	 *
@@ -141,5 +154,10 @@ private:
 	 * Empties the cache.
 	 */
 	void Empty();
+
+	/** 
+	 * Updates our settings from ImgMediaAsettings. 
+	 */
+	void UpdateSettings(const UImgMediaSettings* Settings);
 };
 
diff --git a/Engine/Plugins/Media/ImgMedia/Source/ImgMedia/Private/Scheduler/ImgMediaScheduler.cpp b/Engine/Plugins/Media/ImgMedia/Source/ImgMedia/Private/Scheduler/ImgMediaScheduler.cpp
--- a/Engine/Plugins/Media/ImgMedia/Source/ImgMedia/Private/Scheduler/ImgMediaScheduler.cpp
+++ b/Engine/Plugins/Media/ImgMedia/Source/ImgMedia/Private/Scheduler/ImgMediaScheduler.cpp
@@ -48,6 +48,24 @@ void FImgMediaScheduler::Initialize()
 		NumWorkers = FPlatformMisc::NumberOfWorkerThreadsToSpawn();
 	}
 
+	uint32 StackSize = GetWorkerStackSize();
+
+	// create worker threads
+	FScopeLock Lock(&CriticalSection);
+
+	while (NumWorkers-- > 0)
+	{
+		CreateWorker(StackSize);
+	}
+
+#if WITH_EDITOR
+	UpdateSettingsHandle = UImgMediaSettings::OnSettingsChanged().AddRaw(this, &FImgMediaScheduler::UpdateSettings);
+#endif
+}
+
+
+uint32 FImgMediaScheduler::GetWorkerStackSize()
+{
 	uint32 StackSize = GetDefault<UImgMediaSettings>()->CacheThreadStackSizeKB;
 
 	if (StackSize < IMGMEDIA_SCHEDULERTHREAD_MIN_STACK_SIZE)
@@ -55,15 +73,15 @@ void FImgMediaScheduler::Initialize()
 		StackSize = IMGMEDIA_SCHEDULERTHREAD_MIN_STACK_SIZE;
 	}
 
-	// create worker threads
-	FScopeLock Lock(&CriticalSection);
+	return StackSize;
+}
 
-	while (NumWorkers-- > 0)
-	{
-		auto Thread = new FImgMediaSchedulerThread(*this, StackSize, TPri_Normal);
-		AvailableThreads.Add(Thread);
-		AllThreads.Add(Thread);
-	}
+
+void FImgMediaScheduler::CreateWorker(uint32 StackSize)
+{
+	FImgMediaSchedulerThread* Thread = new FImgMediaSchedulerThread(*this, StackSize, TPri_Normal);
+	AvailableThreads.Add(Thread);
+	AllThreads.Add(Thread);
 }
 
 
@@ -78,6 +96,10 @@ void FImgMediaScheduler::Shutdown()
 {
 	FScopeLock Lock(&CriticalSection);
 
+#if WITH_EDITOR
+	UImgMediaSettings::OnSettingsChanged().Remove(UpdateSettingsHandle);
+#endif
+
 	// destroy worker threads
 	for (FImgMediaSchedulerThread* Thread : AllThreads)
 	{
@@ -144,43 +166,87 @@ void FImgMediaScheduler::TickFetch(FTimespan DeltaTime, FTimespan Timecode)
 
 IQueuedWork* FImgMediaScheduler::GetWorkOrReturnToPool(FImgMediaSchedulerThread* Thread)
 {
-	FScopeLock Lock(&CriticalSection);
-
-	if ((Thread == nullptr) || FPlatformProcess::SupportsMultithreading())
+#if WITH_EDITOR
+	bool DeleteThread = false;
+#endif
 	{
-		const int32 NumLoaders = Loaders.Num();
+		FScopeLock Lock(&CriticalSection);
 
-		if (NumLoaders > 0)
+		if ((Thread == nullptr) || FPlatformProcess::SupportsMultithreading())
 		{
-			int32 CheckedLoaders = 0;
+			const int32 NumLoaders = Loaders.Num();
 
-			while (CheckedLoaders++ < NumLoaders)
+			if (NumLoaders > 0)
 			{
-				LoaderRoundRobin = (LoaderRoundRobin + 1) % NumLoaders;
-
-				TSharedPtr<FImgMediaLoader, ESPMode::ThreadSafe> Loader = Loaders[LoaderRoundRobin].Pin();
+				int32 CheckedLoaders = 0;
 
-				if (Loader.IsValid())
+				while (CheckedLoaders++ < NumLoaders)
 				{
-					IQueuedWork* Work = Loader->GetWork();
+					LoaderRoundRobin = (LoaderRoundRobin + 1) % NumLoaders;
+
+					TSharedPtr<FImgMediaLoader, ESPMode::ThreadSafe> Loader = Loaders[LoaderRoundRobin].Pin();
 
-					if (Work != nullptr)
+					if (Loader.IsValid())
 					{
-						return Work;
+						IQueuedWork* Work = Loader->GetWork();
+
+						if (Work != nullptr)
+						{
+							return Work;
+						}
+					}
+					else
+					{
+						RemoveLoader(LoaderRoundRobin);
 					}
 				}
-				else
-				{
-					RemoveLoader(LoaderRoundRobin);
-				}
 			}
 		}
+
+		if (Thread != nullptr)
+		{
+#if WITH_EDITOR
+			// If there are too many workers then delete this one.
+			int32 NumWorkers = GetDefault<UImgMediaSettings>()->CacheThreads;
+			if (NumWorkers < AllThreads.Num())
+			{
+				AllThreads.Remove(Thread);
+				DeleteThread = true;
+			}
+			else
+			{
+				AvailableThreads.Add(Thread);
+			}
+#else
+			AvailableThreads.Add(Thread);
+#endif
+		}
 	}
 
-	if (Thread != nullptr)
+#if WITH_EDITOR
+	if (DeleteThread)
 	{
-		AvailableThreads.Add(Thread);
+		delete Thread;
 	}
+#endif
 
 	return nullptr;
+	}
+
+
+#if WITH_EDITOR
+
+void FImgMediaScheduler::UpdateSettings(const UImgMediaSettings* Settings)
+{
+	FScopeLock Lock(&CriticalSection);
+
+	// Create more workers if needed.
+	uint32 StackSize = GetWorkerStackSize();
+	while (Settings->CacheThreads > AllThreads.Num())
+	{
+		CreateWorker(StackSize);
+	}
+
 }
+
+#endif
diff --git a/Engine/Plugins/Media/ImgMedia/Source/ImgMedia/Private/Scheduler/ImgMediaScheduler.h b/Engine/Plugins/Media/ImgMedia/Source/ImgMedia/Private/Scheduler/ImgMediaScheduler.h
--- a/Engine/Plugins/Media/ImgMedia/Source/ImgMedia/Private/Scheduler/ImgMediaScheduler.h
+++ b/Engine/Plugins/Media/ImgMedia/Source/ImgMedia/Private/Scheduler/ImgMediaScheduler.h
@@ -11,6 +11,7 @@
 
 class FImgMediaLoader;
 class FImgMediaSchedulerThread;
+class UImgMediaSettings;
 
 
 class FImgMediaScheduler
@@ -94,4 +95,17 @@ private:
 
 	/** Index of the last loader that provided a work item. */
 	int32 LoaderRoundRobin;
+
+	uint32 GetWorkerStackSize();
+	void CreateWorker(uint32 StackSize);
+
+#if WITH_EDITOR
+
+	FDelegateHandle UpdateSettingsHandle;
+	/**
+	 * Updates our settings from ImgMediaAsettings.
+	 */
+	void UpdateSettings(const UImgMediaSettings* Settings);
+
+#endif
 };
diff --git a/Engine/Plugins/Media/ImgMedia/Source/ImgMediaFactory/Private/ImgMediaSettings.cpp b/Engine/Plugins/Media/ImgMedia/Source/ImgMediaFactory/Private/ImgMediaSettings.cpp
--- a/Engine/Plugins/Media/ImgMedia/Source/ImgMediaFactory/Private/ImgMediaSettings.cpp
+++ b/Engine/Plugins/Media/ImgMedia/Source/ImgMediaFactory/Private/ImgMediaSettings.cpp
@@ -18,3 +18,17 @@ UImgMediaSettings::UImgMediaSettings()
 	, DefaultProxy(TEXT("proxy"))
 	, UseDefaultProxy(false)
 { }
+
+#if WITH_EDITOR
+void UImgMediaSettings::PostEditChangeProperty(FPropertyChangedEvent& PropertyChangedEvent)
+{
+	SettingsChangedDelegate.Broadcast(this);
+}
+
+UImgMediaSettings::FOnImgMediaSettingsChanged& UImgMediaSettings::OnSettingsChanged()
+{
+	return SettingsChangedDelegate;
+}
+
+UImgMediaSettings::FOnImgMediaSettingsChanged UImgMediaSettings::SettingsChangedDelegate;
+#endif
diff --git a/Engine/Plugins/Media/ImgMedia/Source/ImgMediaFactory/Public/ImgMediaSettings.h b/Engine/Plugins/Media/ImgMedia/Source/ImgMediaFactory/Public/ImgMediaSettings.h
--- a/Engine/Plugins/Media/ImgMedia/Source/ImgMediaFactory/Public/ImgMediaSettings.h
+++ b/Engine/Plugins/Media/ImgMedia/Source/ImgMediaFactory/Public/ImgMediaSettings.h
@@ -80,6 +80,19 @@ public:
 		return FString();
 	}
 
+#if WITH_EDITOR
+	virtual void PostEditChangeProperty(FPropertyChangedEvent& PropertyChangedEvent) override;
+
+public:
+	DECLARE_MULTICAST_DELEGATE_OneParam(FOnImgMediaSettingsChanged, const UImgMediaSettings*);
+
+	/** Gets a multicast delegate which is called whenever one of the parameters in this settings object changes. */
+	static FOnImgMediaSettingsChanged& OnSettingsChanged();
+
+protected:
+	static FOnImgMediaSettingsChanged SettingsChangedDelegate;
+#endif
+
 private:
 
 	/**
