MediaFramework: Fixed compile issue.

diff --git a/Engine/Plugins/Media/ImgMedia/Source/ImgMedia/Private/GlobalCache/ImgMediaGlobalCache.h b/Engine/Plugins/Media/ImgMedia/Source/ImgMedia/Private/GlobalCache/ImgMediaGlobalCache.h
--- a/Engine/Plugins/Media/ImgMedia/Source/ImgMedia/Private/GlobalCache/ImgMediaGlobalCache.h
+++ b/Engine/Plugins/Media/ImgMedia/Source/ImgMedia/Private/GlobalCache/ImgMediaGlobalCache.h
@@ -4,6 +4,7 @@
 
 #include "CoreTypes.h"
 #include "Containers/Map.h"
+#include "Delegates/IDelegateInstance.h"
 #include "Templates/SharedPointer.h"
 
 struct FImgMediaFrame;
diff --git a/Engine/Plugins/Media/ImgMedia/Source/ImgMedia/Private/Scheduler/ImgMediaScheduler.h b/Engine/Plugins/Media/ImgMedia/Source/ImgMedia/Private/Scheduler/ImgMediaScheduler.h
--- a/Engine/Plugins/Media/ImgMedia/Source/ImgMedia/Private/Scheduler/ImgMediaScheduler.h
+++ b/Engine/Plugins/Media/ImgMedia/Source/ImgMedia/Private/Scheduler/ImgMediaScheduler.h
@@ -4,6 +4,7 @@
 
 #include "Containers/Array.h"
 #include "CoreTypes.h"
+#include "Delegates/IDelegateInstance.h"
 #include "HAL/CriticalSection.h"
 #include "IImgMediaScheduler.h"
 #include "IMediaClockSink.h"
