ImgMediaPlayer update to "V2 player" specs

diff --git a/Engine/Plugins/Media/ImgMedia/Source/ImgMedia/Private/Loader/ImgMediaLoader.cpp b/Engine/Plugins/Media/ImgMedia/Source/ImgMedia/Private/Loader/ImgMediaLoader.cpp
--- a/Engine/Plugins/Media/ImgMedia/Source/ImgMedia/Private/Loader/ImgMediaLoader.cpp
+++ b/Engine/Plugins/Media/ImgMedia/Source/ImgMedia/Private/Loader/ImgMediaLoader.cpp
@@ -52,6 +52,7 @@ FImgMediaLoader::FImgMediaLoader(const TSharedRef<FImgMediaScheduler, ESPMode::T
 	, LastRequestedFrame(INDEX_NONE)
 	, UseGlobalCache(false)
 {
+	ResetFetchLogic();
 	UE_LOG(LogImgMedia, Verbose, TEXT("Loader %p: Created"), this);
 }
 
@@ -112,9 +113,10 @@ void FImgMediaLoader::GetCompletedTimeRanges(TRangeSet<FTimespan>& OutRangeSet)
 }
 
 
+//note: use with V1 player version only!
 TSharedPtr<FImgMediaTextureSample, ESPMode::ThreadSafe> FImgMediaLoader::GetFrameSample(FTimespan Time)
 {
-	const int32 FrameIndex = TimeToFrameNumber(Time);
+	const uint32 FrameIndex = TimeToFrameNumber(Time);
 
 	if (FrameIndex == INDEX_NONE)
 	{
@@ -144,7 +146,7 @@ TSharedPtr<FImgMediaTextureSample, ESPMode::ThreadSafe> FImgMediaLoader::GetFram
 
 	auto Sample = MakeShared<FImgMediaTextureSample, ESPMode::ThreadSafe>();
 	
-	if (!Sample->Initialize(*Frame->Get(), SequenceDim, FrameStartTime, NextStartTime - FrameStartTime))
+	if (!Sample->Initialize(*Frame->Get(), SequenceDim, FMediaTimeStamp(FrameStartTime, 0), NextStartTime - FrameStartTime))
 	{
 		return nullptr;
 	}
@@ -153,6 +155,301 @@ TSharedPtr<FImgMediaTextureSample, ESPMode::ThreadSafe> FImgMediaLoader::GetFram
 }
 
 
+void FImgMediaLoader::ResetFetchLogic()
+{
+	QueuedSampleFetch.LastFrameIndex = INDEX_NONE;
+	// note: we can reset the sequence index here as this will be called when MFW does flush any queues - so we can start from scratch with no issues
+	QueuedSampleFetch.CurrentSequenceIndex = 0;
+}
+
+
+float FImgMediaLoader::FindMaxOverlapInRange(int32 StartIndex, int32 EndIndex, FTimespan StartTime, FTimespan EndTime, int32 & MaxIdx) const
+{
+	int32 IdxInc = 0;
+	if (StartIndex < EndIndex)
+	{
+		IdxInc = 1;
+	}
+	else
+	{
+		IdxInc = -1;
+	}
+
+	// Find index that overlaps the most and is furthest along the timeline...
+	MaxIdx = -1;
+	float MaxOverlap = 0.0f;
+	for (uint32 Idx = StartIndex; Idx != (EndIndex + IdxInc); Idx+=IdxInc)
+	{
+		float Overlap = GetFrameOverlap(Idx, StartTime, EndTime);
+		if (MaxOverlap < Overlap)
+		{
+			MaxOverlap = Overlap;
+			MaxIdx = Idx;
+		}
+	}
+	return MaxOverlap;
+}
+
+
+const TSharedPtr<FImgMediaFrame, ESPMode::ThreadSafe>* FImgMediaLoader::GetFrameForBestIndex(int32 & MaxIdx, int32 LastIndex)
+{
+	const TSharedPtr<FImgMediaFrame, ESPMode::ThreadSafe>* Frame = nullptr;
+	int32 IdxInc = (MaxIdx > LastIndex) ? -1 : 1;
+
+	while (MaxIdx != (LastIndex + IdxInc))
+	{
+		Frame = UseGlobalCache ? GlobalCache->FindAndTouch(SequenceName, MaxIdx) : Frames.FindAndTouch(MaxIdx);
+		if (Frame)
+		{
+			break;
+		}
+		MaxIdx += IdxInc;
+	}
+	return Frame;
+}
+
+
+IMediaSamples::EFetchBestSampleResult FImgMediaLoader::FetchBestVideoSampleForTimeRange(const TRange<FMediaTimeStamp>& TimeRange, TSharedPtr<IMediaTextureSample, ESPMode::ThreadSafe>& OutSample, bool bIsLoopingEnabled, float PlayRate, bool bPlaybackIsBlocking)
+{
+	if (IsInitialized() && TimeRange.HasLowerBound() && TimeRange.HasUpperBound())
+	{
+		FTimespan StartTime = TimeRange.GetLowerBoundValue().Time;
+		FTimespan EndTime = TimeRange.GetUpperBoundValue().Time;
+
+		if (bIsLoopingEnabled)
+		{
+			// Modulo with sequence duration to take care of looping.
+			StartTime = ModuloTime(StartTime);
+			EndTime = ModuloTime(EndTime);
+		}
+
+		// Get start and end frame indices for this time range.
+		int32 StartIndex = TimeToFrameNumber(StartTime);
+		int32 EndIndex = TimeToFrameNumber(EndTime);
+
+		// Sanity checks on returned indices...
+		if ((uint32)StartIndex == INDEX_NONE && (uint32)EndIndex == INDEX_NONE)
+		{
+			return IMediaSamples::EFetchBestSampleResult::NoSample;
+		}
+
+		if ((uint32)StartIndex == INDEX_NONE)
+		{
+			StartIndex = 0;
+		}
+		else if ((uint32)EndIndex == INDEX_NONE)
+		{
+			EndIndex = ImagePaths.Num() - 1;
+		}
+
+		// Find the frame that overlaps the most with the given range & is furthest along on the timeline
+		int32 MaxIdx = -1;
+		if (PlayRate >= 0.0f)
+		{
+			// Forward...
+			if (StartIndex > EndIndex)
+			{
+				int32 MaxIdx1, MaxIdx2;
+				float MaxOverlap1 = FindMaxOverlapInRange(StartIndex, ImagePaths.Num()-1, StartTime, FrameNumberToTime(ImagePaths.Num()), MaxIdx1);
+				float MaxOverlap2 = FindMaxOverlapInRange(0, EndIndex, FTimespan::Zero(), EndTime, MaxIdx2);
+				MaxIdx = (MaxOverlap2 >= MaxOverlap1) ? MaxIdx2 : MaxIdx1;
+			}
+			else
+			{
+				FindMaxOverlapInRange(StartIndex, EndIndex, StartTime, EndTime, MaxIdx);
+			}
+		}
+		else
+		{
+			// Backward...
+			if (StartIndex > EndIndex)
+			{
+				int32 MaxIdx1, MaxIdx2;
+				float MaxOverlap1 = FindMaxOverlapInRange(EndIndex, 0, FTimespan::Zero(), EndTime, MaxIdx1);
+				float MaxOverlap2 = FindMaxOverlapInRange(ImagePaths.Num()-1, StartIndex,  StartTime, FrameNumberToTime(ImagePaths.Num()), MaxIdx2);
+				MaxIdx = (MaxOverlap2 >= MaxOverlap1) ? MaxIdx2 : MaxIdx1;
+			}
+			else
+			{
+				FindMaxOverlapInRange(EndIndex, StartIndex, StartTime, EndTime, MaxIdx);
+			}
+		}
+
+		// Anything?
+		if (MaxIdx >= 0)
+		{
+			const TSharedPtr<FImgMediaFrame, ESPMode::ThreadSafe>* Frame;
+
+			// Request data for the frame we would like... (in case it's not in, yet)
+			RequestFrame(FrameNumberToTime(MaxIdx), PlayRate, bIsLoopingEnabled);
+
+			// If playback is not blocking, we expect less expectancy of precision on the users side, but more need for speedy return of "some ok frame"
+			// So: if we detect non-blocking playback we return a "as good sample as we can", but not always the "perfect" one we calculated
+			// (still we adhere to a rough emulation of a classic output pipeline as other players have)
+			if (!bPlaybackIsBlocking)
+			{
+				// Check what data we actually have in the cache already & attempt to go further backwards on the time line
+				// to get a less then optimal frame that is still on screen and available...
+				if (PlayRate >= 0.0f)
+				{
+					// Forward...
+					if (StartIndex > EndIndex)
+					{
+						if (MaxIdx < StartIndex)
+						{
+							Frame = GetFrameForBestIndex(MaxIdx, 0);
+							if (!Frame)
+							{
+								MaxIdx = ImagePaths.Num() - 1;
+								Frame = GetFrameForBestIndex(MaxIdx, StartIndex);
+							}
+
+						}
+						else
+						{
+							Frame = GetFrameForBestIndex(MaxIdx, StartIndex);
+						}
+					}
+					else
+					{
+						Frame = GetFrameForBestIndex(MaxIdx, StartIndex);
+					}
+				}
+				else
+				{
+					// Backward...
+					if (StartIndex > EndIndex)
+					{
+						if (MaxIdx > EndIndex)
+						{
+							Frame = GetFrameForBestIndex(MaxIdx, ImagePaths.Num() - 1);
+							if (!Frame)
+							{
+								MaxIdx = 0;
+								Frame = GetFrameForBestIndex(MaxIdx, EndIndex);
+							}
+
+						}
+						else
+						{
+							Frame = GetFrameForBestIndex(MaxIdx, EndIndex);
+						}
+					}
+					else
+					{
+						Frame = GetFrameForBestIndex(MaxIdx, EndIndex);
+					}
+				}
+			}
+			else
+			{
+				// Get a frame if we have one available right now...
+				Frame = UseGlobalCache ? GlobalCache->FindAndTouch(SequenceName, MaxIdx) : Frames.FindAndTouch(MaxIdx);
+			}
+
+			// Got a potential frame?
+			if (Frame)
+			{
+				// Different from the last one we returned?
+				if (QueuedSampleFetch.LastFrameIndex != MaxIdx)
+				{
+					// Yes. First time (after flush)?
+					if (QueuedSampleFetch.LastFrameIndex != INDEX_NONE)
+					{
+						// No. Check if we looped and need to start a new sequence...
+						if (PlayRate >= 0.0f && QueuedSampleFetch.LastFrameIndex > MaxIdx)
+						{
+							++QueuedSampleFetch.CurrentSequenceIndex;
+						}
+						else if (PlayRate < 0.0f && QueuedSampleFetch.LastFrameIndex < MaxIdx)
+						{
+							--QueuedSampleFetch.CurrentSequenceIndex;
+						}
+					}
+					QueuedSampleFetch.LastFrameIndex = MaxIdx;
+
+					// We are clear to return it as new result... Make a sample & initialize it...
+					auto Sample = MakeShared<FImgMediaTextureSample, ESPMode::ThreadSafe>();
+					auto Duration = Frame->Get()->Info.FrameRate.AsInterval();
+					if (Sample->Initialize(*Frame->Get(), SequenceDim, FMediaTimeStamp(FrameNumberToTime(MaxIdx), QueuedSampleFetch.CurrentSequenceIndex), FTimespan::FromSeconds(Duration)))
+					{
+						OutSample = Sample;
+						return IMediaSamples::EFetchBestSampleResult::Ok;
+					}
+				}
+			}
+		}
+	}
+	return IMediaSamples::EFetchBestSampleResult::NoSample;
+}
+
+
+bool FImgMediaLoader::PeekVideoSampleTime(FMediaTimeStamp &TimeStamp, bool bIsLoopingEnabled, float PlayRate, const FTimespan& CurrentTime)
+{
+	if (IsInitialized())
+	{
+		int32 Idx;
+		bool bNewSeq = false;
+
+		// Do we know which index we handed out last?
+		if (QueuedSampleFetch.LastFrameIndex != INDEX_NONE)
+		{
+			// Yes. A queue would now yield the next frame (independent of rate). See which index that would be...
+			Idx = int32(QueuedSampleFetch.LastFrameIndex + FMath::Sign(PlayRate));
+			int32 NumFrames = ImagePaths.Num();
+			if (bIsLoopingEnabled)
+			{
+				if (Idx < 0)
+				{
+					Idx = NumFrames - 1;
+					bNewSeq = true;
+				}
+				else if (Idx >= NumFrames)
+				{
+					Idx = 0;
+					bNewSeq = true;
+				}
+			}
+			else
+			{
+				// If we reach either end of the sequence with no looping, we have no more frames to offer
+				if (Idx < 0 || Idx >= NumFrames)
+				{
+					return false;
+				}
+			}
+		}
+		else
+		{
+			// No, we don't have an index. Just compute things based on the current time given...
+			Idx = TimeToFrameNumber(CurrentTime);
+		}
+
+		// If possible, fetch any existing frame data...
+		// (just to see if we have any)
+		const TSharedPtr<FImgMediaFrame, ESPMode::ThreadSafe>* Frame = ((uint32)Idx != INDEX_NONE) ? (UseGlobalCache ? GlobalCache->FindAndTouch(SequenceName, Idx) : Frames.FindAndTouch(Idx)) : nullptr;
+
+		// Start time of this frame
+		FTimespan FrameStart = FrameNumberToTime(Idx);
+
+		// Data is present?
+		if (Frame)
+		{
+			// Yes. Return the timing information...
+			TimeStamp.Time = FrameStart;
+			TimeStamp.SequenceIndex = bNewSeq ? (QueuedSampleFetch.CurrentSequenceIndex + 1) : QueuedSampleFetch.CurrentSequenceIndex;
+			return true;
+		}
+		else
+		{
+			// No data. We will request it (so, like other players do) we fill our (virtual) queue at the current location automatically
+			RequestFrame(FrameStart, PlayRate, bIsLoopingEnabled);
+		}
+	}
+	return false;
+}
+
+
 void FImgMediaLoader::GetPendingTimeRanges(TRangeSet<FTimespan>& OutRangeSet) const
 {
 	FScopeLock Lock(&CriticalSection);
@@ -203,6 +500,9 @@ bool FImgMediaLoader::RequestFrame(FTimespan Time, float PlayRate, bool Loop)
 
 	if ((FrameNumber == INDEX_NONE) || (FrameNumber == LastRequestedFrame))
 	{
+		// Make sure we call the reader even if we do no update - just in case it does anything
+		Reader->OnTick();
+
 		UE_LOG(LogImgMedia, VeryVerbose, TEXT("Loader %p: Skipping frame %i for time %s"), this, FrameNumber, *Time.ToString(TEXT("%h:%m:%s.%t")));
 		return false;
 	}
@@ -380,12 +680,20 @@ void FImgMediaLoader::LoadSequence(const FString& SequencePath, const FFrameRate
 
 uint32 FImgMediaLoader::TimeToFrameNumber(FTimespan Time) const
 {
-	if ((Time < FTimespan::Zero()) || (Time > SequenceDuration))
+	if ((Time < FTimespan::Zero()) || (Time >= SequenceDuration))
 	{
 		return INDEX_NONE;
 	}
 
-	return (Time.GetTicks() * SequenceFrameRate.Numerator) / (ETimespan::TicksPerSecond * SequenceFrameRate.Denominator);
+	const double FrameTimeErrorTollerance = 0.0001;
+
+	// note: we snap to the next best whole frame index if the compute result is with in FrameTimeErrorTollerance * FrameDuration to avoid
+	// incorrect frame selection if the value passed in is just ever so slightly off
+	double FrameDuration = (double)SequenceFrameRate.Numerator / SequenceFrameRate.Denominator;
+	double Frame = Time.GetTotalSeconds() * FrameDuration;
+	double Epsilon = FrameTimeErrorTollerance * FrameDuration;
+
+	return uint32(Frame + Epsilon);
 }
 
 
@@ -550,3 +858,43 @@ void FImgMediaLoader::NotifyWorkComplete(FImgMediaLoaderWork& CompletedWork, int
 
 	WorkPool.Push(&CompletedWork);
 }
+
+
+FTimespan FImgMediaLoader::ModuloTime(FTimespan Time)
+{
+	bool IsNegative = Time < FTimespan::Zero();
+	
+	FTimespan NewTime = Time % SequenceDuration;
+	if (IsNegative)
+	{
+		NewTime = SequenceDuration + NewTime;
+	}
+
+	return NewTime;
+}
+
+
+float FImgMediaLoader::GetFrameOverlap(uint32 FrameIndex, FTimespan StartTime, FTimespan EndTime) const
+{
+	check(StartTime <= EndTime);
+	if (StartTime == EndTime)
+	{
+		return 0.0f;
+	}
+
+	float Overlap = 0.0f;
+
+	// Set up ranges.
+	FTimespan FrameStartTime = FrameNumberToTime(FrameIndex);
+	FTimespan FrameEndTime = FrameStartTime + FrameNumberToTime(1);
+
+	TRange<FTimespan> FrameRange(FrameStartTime, FrameEndTime);
+	TRange<FTimespan> TimeRange(StartTime, EndTime);
+	TRange<FTimespan> OverlapRange = TRange<FTimespan>::Intersection(FrameRange, TimeRange);
+
+	// Get overlap size.
+	FTimespan OverlapTimespan = OverlapRange.Size<FTimespan>();
+	Overlap = OverlapTimespan.GetTotalSeconds();
+
+	return Overlap;
+}
diff --git a/Engine/Plugins/Media/ImgMedia/Source/ImgMedia/Private/Loader/ImgMediaLoader.h b/Engine/Plugins/Media/ImgMedia/Source/ImgMedia/Private/Loader/ImgMediaLoader.h
--- a/Engine/Plugins/Media/ImgMedia/Source/ImgMedia/Private/Loader/ImgMediaLoader.h
+++ b/Engine/Plugins/Media/ImgMedia/Source/ImgMedia/Private/Loader/ImgMediaLoader.h
@@ -5,6 +5,7 @@
 #include "CoreMinimal.h"
 #include "Containers/LruCache.h"
 #include "Containers/Queue.h"
+#include "IMediaSamples.h"
 #include "Misc/FrameRate.h"
 #include "Templates/SharedPointer.h"
 
@@ -12,6 +13,7 @@ class FImgMediaGlobalCache;
 class FImgMediaLoaderWork;
 class FImgMediaScheduler;
 class FImgMediaTextureSample;
+class FMediaTimeStamp;
 class IImageWrapperModule;
 class IImgMediaReader;
 class IQueuedWork;
@@ -84,6 +86,29 @@ public:
 		return Info;
 	}
 
+	/**
+	 * Tries to get the best sample for a given time range.
+	 *
+	 * @param TimeRange is the range to check for.
+	 * @param OutSample will be filled in with the sample if found.
+	 * @param bIsLoopingEnabled True if we can loop.
+	 * @param PlayRate How fast we are playing.
+	 * @param Facade does use blocking playback to fetch samples
+	 * @return True if successful.
+	 */
+	IMediaSamples::EFetchBestSampleResult FetchBestVideoSampleForTimeRange(const TRange<FMediaTimeStamp>& TimeRange, TSharedPtr<IMediaTextureSample, ESPMode::ThreadSafe>& OutSample, bool bIsLoopingEnabled, float PlayRate, bool bPlaybackIsBlocking);
+
+	/**
+	 * Checks to see if a sample is available at the specificed time.
+	 *
+	 * @param TimeStamp Will be filled in with the time stamp of the sample if a sample is found.
+	 * @param bIsLoopingEnabled True if we can loop.
+	 * @param PlayRate How fast we are playing.
+	 * @param CurrentTime Time to get a sample for.
+	 * @return True if a sample is available.
+	 */
+	bool PeekVideoSampleTime(FMediaTimeStamp& TimeStamp, bool bIsLoopingEnabled, float PlayRate, const FTimespan& CurrentTime);
+
 	/**
 	 * Get the time ranges of frames that are pending to be loaded.
 	 *
@@ -193,6 +218,11 @@ public:
 	 * @return bool if the frame will be loaded, false otherwise.
 	 */
 	bool RequestFrame(FTimespan Time, float PlayRate, bool Loop);
+	
+	/**
+	 * Reset "queued fetch" related state used to emulate player output queue behavior
+	 */
+	void ResetFetchLogic();
 
 protected:
 
@@ -240,6 +270,33 @@ protected:
 	 */
 	void Update(int32 PlayHeadFrame, float PlayRate, bool Loop);
 
+	/***
+	 * Modulos the time so that it is between 0 and SequenceDuration.
+	 * Handles negative numbers appropriately.
+	 */
+	FTimespan ModuloTime(FTimespan Time);
+
+	/**
+	 * Gets the amount of overlap (in seconds) between a frame and a time range.
+	 * A negative or zero  value indicates the frame does not overlap the range.
+	 *
+	 * @param FrameIndex Index of frame to check.
+	 * @param StartTime Start of the range.
+	 * @param EndTime End of the range.
+	 * @param Amount of overlap.
+	 */
+	float GetFrameOverlap(uint32 FrameIndex, FTimespan StartTime, FTimespan EndTime) const;
+
+	/**
+	 * Find maximum overlapping frame index for given range
+	 */
+	float FindMaxOverlapInRange(int32 StartIndex, int32 EndIndex, FTimespan StartTime, FTimespan EndTime, int32& MaxIdx) const;
+
+	/**
+	 * Get frame data for given index. If not available attempt to find earlier frame up to given index
+	 */
+	const TSharedPtr<FImgMediaFrame, ESPMode::ThreadSafe>* GetFrameForBestIndex(int32 & MaxIdx, int32 LastIndex);
+
 private:
 
 	/** Critical section for synchronizing access to Frames. */
@@ -303,4 +360,13 @@ private:
 
 	/** True if we are using the global cache, false to use the local cache. */
 	bool UseGlobalCache;
+
+private:
+	/** State related to "queue style" frame access functions */
+	struct
+	{
+		int32 LastFrameIndex;
+		uint64 CurrentSequenceIndex;
+	} QueuedSampleFetch;
+
 };
diff --git a/Engine/Plugins/Media/ImgMedia/Source/ImgMedia/Private/Player/ImgMediaPlayer.cpp b/Engine/Plugins/Media/ImgMedia/Source/ImgMedia/Private/Player/ImgMediaPlayer.cpp
--- a/Engine/Plugins/Media/ImgMedia/Source/ImgMedia/Private/Player/ImgMediaPlayer.cpp
+++ b/Engine/Plugins/Media/ImgMedia/Source/ImgMedia/Private/Player/ImgMediaPlayer.cpp
@@ -16,6 +16,7 @@
 #include "ImgMediaSettings.h"
 #include "ImgMediaTextureSample.h"
 
+#define IMG_MEDIA_PLAYER_VERSION 2
 #define LOCTEXT_NAMESPACE "FImgMediaPlayer"
 
 
@@ -47,6 +48,7 @@ FImgMediaPlayer::FImgMediaPlayer(IMediaEventSink& InEventSink, const TSharedRef<
 	, ShouldLoop(false)
 	, GlobalCache(InGlobalCache)
 	, RequestFrameHasRun(true)
+	, PlaybackIsBlocking(false)
 { }
 
 
@@ -80,6 +82,7 @@ void FImgMediaPlayer::Close()
 	LastFetchTime = FTimespan::MinValue();
 	PlaybackRestarted = false;
 	SelectedVideoTrack = INDEX_NONE;
+	PlaybackIsBlocking = false;
 
 	EventSink.ReceiveMediaEvent(EMediaEvent::TracksChanged);
 	EventSink.ReceiveMediaEvent(EMediaEvent::MediaClosed);
@@ -242,6 +245,8 @@ void FImgMediaPlayer::TickInput(FTimespan DeltaTime, FTimespan /*Timecode*/)
 		}
 	}
 
+#if IMG_MEDIA_PLAYER_VERSION == 1
+
 	if ((CurrentState != EMediaState::Playing) || (CurrentDuration == FTimespan::Zero()))
 	{
 		return; // nothing to play
@@ -308,10 +313,32 @@ void FImgMediaPlayer::TickInput(FTimespan DeltaTime, FTimespan /*Timecode*/)
 		Loader->RequestFrame(CurrentTime, CurrentRate, ShouldLoop);
 	}
 	RequestFrameHasRun = true;
+#endif // IMG_MEDIA_PLAYER_VERSION == 1
+}
+
+bool FImgMediaPlayer::FlushOnSeekStarted() const
+{
+#if IMG_MEDIA_PLAYER_VERSION == 1
+	return false;
+#else
+	// Flush on start, otherwise if we already have the frames ready during a seek,
+	// the frames we return will get flushed straight away.
+	return true;
+#endif
+}
+
+bool FImgMediaPlayer::FlushOnSeekCompleted() const
+{
+#if IMG_MEDIA_PLAYER_VERSION == 1
+	return true;
+#else
+	return false;
+#endif
 }
 
 void FImgMediaPlayer::ProcessVideoSamples()
 {
+#if IMG_MEDIA_PLAYER_VERSION == 1
 	// Did we already run this frame?
 	if (RequestFrameHasRun)
 	{
@@ -332,6 +359,25 @@ void FImgMediaPlayer::ProcessVideoSamples()
 			Scheduler->TickFetch(FTimespan::Zero(), FTimespan::Zero());
 		}
 	}
+#endif // IMG_MEDIA_PLAYER_VERSION == 1
+}
+
+//-----------------------------------------------------------------------------
+/**
+	Get special feature flags states
+*/
+bool FImgMediaPlayer::GetPlayerFeatureFlag(EFeatureFlag flag) const
+{
+#if IMG_MEDIA_PLAYER_VERSION >= 2
+	switch (flag)
+	{
+	case EFeatureFlag::UsePlaybackTimingV2:
+		return true;
+	default:
+		break;
+	}
+#endif  // IMG_MEDIA_PLAYER_VERSION >= 2
+	return IMediaPlayer::GetPlayerFeatureFlag(flag);
 }
 
 
@@ -482,6 +528,7 @@ bool FImgMediaPlayer::Seek(const FTimespan& Time)
 		CurrentState = EMediaState::Paused;
 	}
 
+#if IMG_MEDIA_PLAYER_VERSION == 1
 	// more timing hacks for Sequencer
 	CurrentTime = Time + HackDeltaTimeOffset;
 	DeltaTimeHackApplied = true;
@@ -490,8 +537,10 @@ bool FImgMediaPlayer::Seek(const FTimespan& Time)
 	{
 		CurrentTime -= HackDeltaTimeOffset;
 	}
+#else
+	CurrentTime = Time;
+#endif // IMG_MEDIA_PLAYER_VERSION == 1
 
-	// update the loader
 	if (CurrentState == EMediaState::Paused)
 	{
 		Loader->RequestFrame(CurrentTime, CurrentRate, ShouldLoop);
@@ -536,10 +585,12 @@ bool FImgMediaPlayer::SetRate(float Rate)
 	{
 		if (CurrentState == EMediaState::Stopped)
 		{
+#if IMG_MEDIA_PLAYER_VERSION == 1
 			if (Rate < 0.0f)
 			{
 				CurrentTime = CurrentDuration - FTimespan(1);
 			}
+#endif // IMG_MEDIA_PLAYER_VERSION == 1
 
 			PlaybackRestarted = true;
 		}
@@ -559,7 +610,6 @@ bool FImgMediaPlayer::SetRate(float Rate)
 		CurrentState = EMediaState::Paused;
 
 		EventSink.ReceiveMediaEvent(EMediaEvent::PlaybackSuspended);
-
 		return true;
 	}
 
@@ -568,12 +618,17 @@ bool FImgMediaPlayer::SetRate(float Rate)
 	return true;
 }
 
+void FImgMediaPlayer::SetBlockingPlaybackHint(bool FacadeWillUseBlockingPlayback)
+{
+	PlaybackIsBlocking = FacadeWillUseBlockingPlayback;
+}
 
 /* IMediaSamples interface
  *****************************************************************************/
 
 bool FImgMediaPlayer::FetchVideo(TRange<FTimespan> TimeRange, TSharedPtr<IMediaTextureSample, ESPMode::ThreadSafe>& OutSample)
 {
+#if IMG_MEDIA_PLAYER_VERSION == 1
 	if ((CurrentState != EMediaState::Paused) && (CurrentState != EMediaState::Playing))
 	{
 		return false; // nothing to play
@@ -602,18 +657,46 @@ bool FImgMediaPlayer::FetchVideo(TRange<FTimespan> TimeRange, TSharedPtr<IMediaT
 	OutSample = Sample;
 
 	return true;
+#else // IMG_MEDIA_PLAYER_VERSION == 1
+	return false;
+#endif // IMG_MEDIA_PLAYER_VERSION == 1
 }
 
 
 void FImgMediaPlayer::FlushSamples()
 {
-	LastFetchTime = FTimespan::MinValue();
+	if (IsInitialized())
+	{
+		LastFetchTime = FTimespan::MinValue();
+		Loader->ResetFetchLogic();
+	}
+}
+
+
+IMediaSamples::EFetchBestSampleResult FImgMediaPlayer::FetchBestVideoSampleForTimeRange(const TRange<FMediaTimeStamp>& TimeRange, TSharedPtr<IMediaTextureSample, ESPMode::ThreadSafe>& OutSample, bool bReverse)
+{
+	IMediaSamples::EFetchBestSampleResult SampleResult = EFetchBestSampleResult::NoSample;
+
+	if (Loader.IsValid())
+	{
+		// See if we have any samples in the specified time range.
+		SampleResult = Loader->FetchBestVideoSampleForTimeRange(TimeRange, OutSample, ShouldLoop, CurrentRate, PlaybackIsBlocking);
+		if (SampleResult == IMediaSamples::EFetchBestSampleResult::Ok)
+		{
+			CurrentTime = OutSample->GetTime().Time;
+		}
+	}
+	return SampleResult;
 }
 
 
-bool FImgMediaPlayer::PeekVideoSampleTime(FMediaTimeStamp & TimeStamp)
+bool FImgMediaPlayer::PeekVideoSampleTime(FMediaTimeStamp& TimeStamp)
 {
-	// player does not support v2 timing control at this point -> no need for this method, yet
+	if (Loader.IsValid())
+	{
+		// Do we have the current frame?
+		return Loader->PeekVideoSampleTime(TimeStamp, ShouldLoop, CurrentRate, GetTime());
+	}
 	return false;
 }
 
diff --git a/Engine/Plugins/Media/ImgMedia/Source/ImgMedia/Private/Player/ImgMediaPlayer.h b/Engine/Plugins/Media/ImgMedia/Source/ImgMedia/Private/Player/ImgMediaPlayer.h
--- a/Engine/Plugins/Media/ImgMedia/Source/ImgMedia/Private/Player/ImgMediaPlayer.h
+++ b/Engine/Plugins/Media/ImgMedia/Source/ImgMedia/Private/Player/ImgMediaPlayer.h
@@ -60,7 +60,10 @@ public:
 	virtual bool Open(const FString& Url, const IMediaOptions* Options) override;
 	virtual bool Open(const TSharedRef<FArchive, ESPMode::ThreadSafe>& Archive, const FString& OriginalUrl, const IMediaOptions* Options) override;
 	virtual void TickInput(FTimespan DeltaTime, FTimespan Timecode) override;
+	virtual bool FlushOnSeekStarted() const;
+	virtual bool FlushOnSeekCompleted() const;
 	virtual void ProcessVideoSamples() override;
+	virtual bool GetPlayerFeatureFlag(EFeatureFlag flag) const override;
 
 protected:
 
@@ -92,6 +95,7 @@ protected:
 	virtual bool Seek(const FTimespan& Time) override;
 	virtual bool SetLooping(bool Looping) override;
 	virtual bool SetRate(float Rate) override;
+	virtual void SetBlockingPlaybackHint(bool bFacadeWillUseBlockingPlayback) override;
 
 protected:
 
@@ -99,6 +103,8 @@ protected:
 
 	virtual bool FetchVideo(TRange<FTimespan> TimeRange, TSharedPtr<IMediaTextureSample, ESPMode::ThreadSafe>& OutSample) override;
 	virtual void FlushSamples() override;
+
+	virtual EFetchBestSampleResult FetchBestVideoSampleForTimeRange(const TRange<FMediaTimeStamp> & TimeRange, TSharedPtr<IMediaTextureSample, ESPMode::ThreadSafe>& OutSample, bool bReverse) override;
 	virtual bool PeekVideoSampleTime(FMediaTimeStamp & TimeStamp) override;
 
 protected:
@@ -163,4 +169,7 @@ private:
 
 	/** True if we have run RequestFrame already for this frame. */
 	bool RequestFrameHasRun;
+
+	/** True if facade has signaled it uses blocking playback, false if not */
+	bool PlaybackIsBlocking;
 };
diff --git a/Engine/Plugins/Media/ImgMedia/Source/ImgMedia/Private/Player/ImgMediaTextureSample.h b/Engine/Plugins/Media/ImgMedia/Source/ImgMedia/Private/Player/ImgMediaTextureSample.h
--- a/Engine/Plugins/Media/ImgMedia/Source/ImgMedia/Private/Player/ImgMediaTextureSample.h
+++ b/Engine/Plugins/Media/ImgMedia/Source/ImgMedia/Private/Player/ImgMediaTextureSample.h
@@ -36,7 +36,7 @@ public:
 	 * @param InTime The sample time (in the player's local clock).
 	 * @param InDuration The duration for which the sample is valid.
 	 */
-	bool Initialize(FImgMediaFrame& InFrame, const FIntPoint& InOutputDim, FTimespan InTime, FTimespan InDuration)
+	bool Initialize(FImgMediaFrame& InFrame, const FIntPoint& InOutputDim, FMediaTimeStamp InTime, FTimespan InDuration)
 	{
 		Duration = InDuration;
 		Frame = InFrame;
@@ -94,7 +94,7 @@ public:
 
 	virtual FMediaTimeStamp GetTime() const override
 	{
-		return FMediaTimeStamp(Time);
+		return Time;
 	}
 
 	virtual bool IsCacheable() const override
@@ -119,5 +119,5 @@ private:
 	FIntPoint OutputDim;
 
 	/** Play time for which the sample was generated. */
-	FTimespan Time;
+	FMediaTimeStamp Time;
 };
