Fixing up reworked ImgMediaPlayer to more fully comply with V2 timing in MFW
#rnx

diff --git a/Engine/Plugins/Media/ImgMedia/Source/ImgMedia/Private/Player/ImgMediaPlayer.cpp b/Engine/Plugins/Media/ImgMedia/Source/ImgMedia/Private/Player/ImgMediaPlayer.cpp
--- a/Engine/Plugins/Media/ImgMedia/Source/ImgMedia/Private/Player/ImgMediaPlayer.cpp
+++ b/Engine/Plugins/Media/ImgMedia/Source/ImgMedia/Private/Player/ImgMediaPlayer.cpp
@@ -313,6 +313,9 @@ void FImgMediaPlayer::TickInput(FTimespan DeltaTime, FTimespan /*Timecode*/)
 		Loader->RequestFrame(CurrentTime, CurrentRate, ShouldLoop);
 	}
 	RequestFrameHasRun = true;
+#else
+	// Tick the scheduler an extra time in addition to its hookup as media clock sink, so we also get it moving forward during blocked playback
+	Scheduler->TickInput(FTimespan::Zero(), FTimespan::MinValue());
 #endif // IMG_MEDIA_PLAYER_VERSION == 1
 }
 
@@ -429,14 +432,14 @@ bool FImgMediaPlayer::QueryCacheState(EMediaCacheState State, TRangeSet<FTimespa
 
 bool FImgMediaPlayer::CanControl(EMediaControl Control) const
 {
-	if (!IsInitialized())
+	if (Control == EMediaControl::BlockOnFetch)
 	{
-		return false;
+		return true;
 	}
 
-	if (Control == EMediaControl::BlockOnFetch)
+	if (!IsInitialized())
 	{
-		return ((CurrentState == EMediaState::Paused) || (CurrentState == EMediaState::Playing));
+		return false;
 	}
 
 	if (Control == EMediaControl::Pause)
diff --git a/Engine/Plugins/Media/ImgMedia/Source/ImgMedia/Private/Scheduler/ImgMediaScheduler.cpp b/Engine/Plugins/Media/ImgMedia/Source/ImgMedia/Private/Scheduler/ImgMediaScheduler.cpp
--- a/Engine/Plugins/Media/ImgMedia/Source/ImgMedia/Private/Scheduler/ImgMediaScheduler.cpp
+++ b/Engine/Plugins/Media/ImgMedia/Source/ImgMedia/Private/Scheduler/ImgMediaScheduler.cpp
@@ -138,9 +138,13 @@ void FImgMediaScheduler::RemoveLoader(int32 LoaderIndex)
 
 
 /* IMediaClockSink interface
+* 
+* Note that this is ticked both from its own registration as clock sink with MFW,
+* but also explicitly at the end of any ImgMediaPlayer tick. The latter is done to
+* ensure we schedule jobs also when waiting for blocked playback to resolve.
  *****************************************************************************/
 
-void FImgMediaScheduler::TickFetch(FTimespan DeltaTime, FTimespan Timecode)
+void FImgMediaScheduler::TickInput(FTimespan DeltaTime, FTimespan Timecode)
 {
 	FScopeLock Lock(&CriticalSection);
 
diff --git a/Engine/Plugins/Media/ImgMedia/Source/ImgMedia/Private/Scheduler/ImgMediaScheduler.h b/Engine/Plugins/Media/ImgMedia/Source/ImgMedia/Private/Scheduler/ImgMediaScheduler.h
--- a/Engine/Plugins/Media/ImgMedia/Source/ImgMedia/Private/Scheduler/ImgMediaScheduler.h
+++ b/Engine/Plugins/Media/ImgMedia/Source/ImgMedia/Private/Scheduler/ImgMediaScheduler.h
@@ -63,7 +63,7 @@ public:
 
 	//~ IMediaClockSink interface
 
-	virtual void TickFetch(FTimespan DeltaTime, FTimespan Timecode) override;
+	virtual void TickInput(FTimespan DeltaTime, FTimespan Timecode) override;
 
 protected:
 
