Making sure ImgMediaPlayer only returns samples once fully initialized
#rnx

diff --git a/Engine/Plugins/Media/ImgMedia/Source/ImgMedia/Private/Player/ImgMediaPlayer.cpp b/Engine/Plugins/Media/ImgMedia/Source/ImgMedia/Private/Player/ImgMediaPlayer.cpp
--- a/Engine/Plugins/Media/ImgMedia/Source/ImgMedia/Private/Player/ImgMediaPlayer.cpp
+++ b/Engine/Plugins/Media/ImgMedia/Source/ImgMedia/Private/Player/ImgMediaPlayer.cpp
@@ -680,7 +680,7 @@ IMediaSamples::EFetchBestSampleResult FImgMediaPlayer::FetchBestVideoSampleForTi
 {
 	IMediaSamples::EFetchBestSampleResult SampleResult = EFetchBestSampleResult::NoSample;
 
-	if (Loader.IsValid())
+	if (Loader.IsValid() && IsInitialized())
 	{
 		// See if we have any samples in the specified time range.
 		SampleResult = Loader->FetchBestVideoSampleForTimeRange(TimeRange, OutSample, ShouldLoop, CurrentRate, PlaybackIsBlocking);
@@ -695,7 +695,7 @@ IMediaSamples::EFetchBestSampleResult FImgMediaPlayer::FetchBestVideoSampleForTi
 
 bool FImgMediaPlayer::PeekVideoSampleTime(FMediaTimeStamp& TimeStamp)
 {
-	if (Loader.IsValid())
+	if (Loader.IsValid() && IsInitialized())
 	{
 		// Do we have the current frame?
 		return Loader->PeekVideoSampleTime(TimeStamp, ShouldLoop, CurrentRate, GetTime());
