ImgMedia: Fixed not being able to return samples when looping a video of length 1 frame.
#rnx

diff --git a/Engine/Plugins/Media/ImgMedia/Source/ImgMedia/Private/Loader/ImgMediaLoader.cpp b/Engine/Plugins/Media/ImgMedia/Source/ImgMedia/Private/Loader/ImgMediaLoader.cpp
--- a/Engine/Plugins/Media/ImgMedia/Source/ImgMedia/Private/Loader/ImgMediaLoader.cpp
+++ b/Engine/Plugins/Media/ImgMedia/Source/ImgMedia/Private/Loader/ImgMediaLoader.cpp
@@ -371,23 +371,37 @@ IMediaSamples::EFetchBestSampleResult FImgMediaLoader::FetchBestVideoSampleForTi
 			// Got a potential frame?
 			if (Frame)
 			{
-				// Different from the last one we returned?
-				if (QueuedSampleFetch.LastFrameIndex != MaxIdx)
+				// Yes. First time (after flush)?
+				int32 NewSequenceIndex = QueuedSampleFetch.CurrentSequenceIndex;
+				if (QueuedSampleFetch.LastFrameIndex != INDEX_NONE)
 				{
-					// Yes. First time (after flush)?
-					if (QueuedSampleFetch.LastFrameIndex != INDEX_NONE)
+					// No. Check if we looped and need to start a new sequence...
+					if (PlayRate >= 0.0f && QueuedSampleFetch.LastFrameIndex > MaxIdx)
+					{
+						++NewSequenceIndex;
+					}
+					else if (PlayRate < 0.0f && QueuedSampleFetch.LastFrameIndex < MaxIdx)
+					{
+						--NewSequenceIndex;
+					}
+					else if (ImagePaths.Num() == 1)
 					{
-						// No. Check if we looped and need to start a new sequence...
-						if (PlayRate >= 0.0f && QueuedSampleFetch.LastFrameIndex > MaxIdx)
+						if (PlayRate >= 0.0f)
 						{
-							++QueuedSampleFetch.CurrentSequenceIndex;
+							++NewSequenceIndex;
 						}
-						else if (PlayRate < 0.0f && QueuedSampleFetch.LastFrameIndex < MaxIdx)
+						else
 						{
-							--QueuedSampleFetch.CurrentSequenceIndex;
+							--NewSequenceIndex;
 						}
 					}
+				}
+
+				// Different from the last one we returned?
+				if ((QueuedSampleFetch.LastFrameIndex != MaxIdx) || (QueuedSampleFetch.CurrentSequenceIndex != NewSequenceIndex))
+				{
 					QueuedSampleFetch.LastFrameIndex = MaxIdx;
+					QueuedSampleFetch.CurrentSequenceIndex = NewSequenceIndex;
 
 					// We are clear to return it as new result... Make a sample & initialize it...
 					auto Sample = MakeShared<FImgMediaTextureSample, ESPMode::ThreadSafe>();
