ImgMedia: Fixed up "pre-mipmap" code.
#rnx

diff --git a/Engine/Plugins/Media/ImgMedia/Source/ImgMedia/Private/Loader/ImgMediaLoader.cpp b/Engine/Plugins/Media/ImgMedia/Source/ImgMedia/Private/Loader/ImgMediaLoader.cpp
--- a/Engine/Plugins/Media/ImgMedia/Source/ImgMedia/Private/Loader/ImgMediaLoader.cpp
+++ b/Engine/Plugins/Media/ImgMedia/Source/ImgMedia/Private/Loader/ImgMediaLoader.cpp
@@ -391,7 +391,7 @@ IMediaSamples::EFetchBestSampleResult FImgMediaLoader::FetchBestVideoSampleForTi
 					{
 						--NewSequenceIndex;
 					}
-					else if (ImagePaths.Num() == 1)
+					else if (GetNumImages() == 1)
 					{
 						if (PlayRate >= 0.0f)
 						{
