ImgMedia: Fixed sequence index mismatch for 1 frame videos.
#rnx

diff --git a/Engine/Plugins/Media/ImgMedia/Source/ImgMedia/Private/Loader/ImgMediaLoader.cpp b/Engine/Plugins/Media/ImgMedia/Source/ImgMedia/Private/Loader/ImgMediaLoader.cpp
--- a/Engine/Plugins/Media/ImgMedia/Source/ImgMedia/Private/Loader/ImgMediaLoader.cpp
+++ b/Engine/Plugins/Media/ImgMedia/Source/ImgMedia/Private/Loader/ImgMediaLoader.cpp
@@ -393,19 +393,12 @@ IMediaSamples::EFetchBestSampleResult FImgMediaLoader::FetchBestVideoSampleForTi
 					}
 					else if (GetNumImages() == 1)
 					{
-						if (PlayRate >= 0.0f)
-						{
-							++NewSequenceIndex;
-						}
-						else
-						{
-							--NewSequenceIndex;
-						}
+						NewSequenceIndex = TimeRange.GetLowerBoundValue().SequenceIndex;
 					}
 				}
 
 				// Different from the last one we returned?
-				if ((QueuedSampleFetch.LastFrameIndex != MaxIdx) || (QueuedSampleFetch.CurrentSequenceIndex != NewSequenceIndex))
+				if ((QueuedSampleFetch.LastFrameIndex != MaxIdx) || (QueuedSampleFetch.CurrentSequenceIndex != NewSequenceIndex) || (ImagePaths.Num() == 1))
 				{
 					QueuedSampleFetch.LastFrameIndex = MaxIdx;
 					QueuedSampleFetch.CurrentSequenceIndex = NewSequenceIndex;
