ImgMedia: Added support for 1 and 2 channel EXR files for the GPU reader.
#rnx

diff --git a/Engine/Plugins/Media/ImgMedia/Shaders/Private/ExrSwizzler.usf b/Engine/Plugins/Media/ImgMedia/Shaders/Private/ExrSwizzler.usf
--- a/Engine/Plugins/Media/ImgMedia/Shaders/Private/ExrSwizzler.usf
+++ b/Engine/Plugins/Media/ImgMedia/Shaders/Private/ExrSwizzler.usf
@@ -63,24 +63,36 @@ half4 MainPS(noperspective float4 UVAndScreenPos : TEXCOORD0) : SV_Target0
 	int xCoord = UV.x * TextureWidth;
 	int yCoord = UV.y * TextureHeight;
 
-	// Depending on the EXR file it could contain multi-channel data. We support only 3 or 4 channel data.
-#if IS_RGBA
-	int Channels = 4;
-#else
+	// Depending on the EXR file it could contain multi-channel data.
+#if NUM_CHANNELS == 1
+	int Channels = 1;
+#elif NUM_CHANNELS == 2
+	int Channels = 2;
+#elif NUM_CHANNELS == 3
 	int Channels = 3;
+#else
+	int Channels = 4;
 #endif
 	// Current scanline position.
 	int Offset = yCoord * TextureWidth * Channels + xCoord + (yCoord + 1) * 4;
 	int ChannelIndex = 0;
-#if IS_RGBA
+#if NUM_CHANNELS == 4
 	half A = UintContainingFloatToFloat(GetValueFromTheBuffer(Offset, TextureWidth, TextureHeight, ChannelIndex++));
 #else
 	half A = 1.;
 #endif
 
 	half B = UintContainingFloatToFloat(GetValueFromTheBuffer(Offset, TextureWidth, TextureHeight, ChannelIndex++));
+#if NUM_CHANNELS <= 1
+	half G = B;
+#else
 	half G = UintContainingFloatToFloat(GetValueFromTheBuffer(Offset, TextureWidth, TextureHeight, ChannelIndex++));
+#endif
+#if NUM_CHANNELS <= 2
+	half R = B;
+#else
 	half R = UintContainingFloatToFloat(GetValueFromTheBuffer(Offset, TextureWidth, TextureHeight, ChannelIndex++));
+#endif
 
 	half3 Rgb = half3(R, G, B);
 
diff --git a/Engine/Plugins/Media/ImgMedia/Source/ExrReaderGpu/Public/ExrSwizzlingShader.h b/Engine/Plugins/Media/ImgMedia/Source/ExrReaderGpu/Public/ExrSwizzlingShader.h
--- a/Engine/Plugins/Media/ImgMedia/Source/ExrReaderGpu/Public/ExrSwizzlingShader.h
+++ b/Engine/Plugins/Media/ImgMedia/Source/ExrReaderGpu/Public/ExrSwizzlingShader.h
@@ -40,7 +40,7 @@ class EXRREADERGPU_API FExrSwizzlePS : public FGlobalShader
 	SHADER_USE_PARAMETER_STRUCT(FExrSwizzlePS, FGlobalShader);
 
 	/** If the provided buffer is RGBA the shader would work slightly differently to RGB. */
-	class FRgbaSwizzle : SHADER_PERMUTATION_BOOL("IS_RGBA");
+	class FRgbaSwizzle : SHADER_PERMUTATION_INT("NUM_CHANNELS", 5);
 	using FPermutationDomain = TShaderPermutationDomain<FRgbaSwizzle>;
 
 	static bool ShouldCompilePermutation(const FGlobalShaderPermutationParameters& Parameters) 
diff --git a/Engine/Plugins/Media/ImgMedia/Source/ImgMedia/Private/Readers/ExrImgMediaReaderGpu.cpp b/Engine/Plugins/Media/ImgMedia/Source/ImgMedia/Private/Readers/ExrImgMediaReaderGpu.cpp
--- a/Engine/Plugins/Media/ImgMedia/Source/ImgMedia/Private/Readers/ExrImgMediaReaderGpu.cpp
+++ b/Engine/Plugins/Media/ImgMedia/Source/ImgMedia/Private/Readers/ExrImgMediaReaderGpu.cpp
@@ -201,7 +201,7 @@ bool FExrImgMediaReaderGpu::ReadFrame(int32 FrameId, int32 MipLevel, const FImgM
 		}
 	}
 
-	OutFrame->Format = NumChannels == 3 ? EMediaTextureSampleFormat::FloatRGB : EMediaTextureSampleFormat::FloatRGBA;
+	OutFrame->Format = NumChannels <= 3 ? EMediaTextureSampleFormat::FloatRGB : EMediaTextureSampleFormat::FloatRGBA;
 	OutFrame->Stride = LargestDim.X * PixelSize;
 	auto RenderThreadSwizzler = [this, BufferDataArray, LargestDim, FrameId, NumChannels, NumMipLevels](FRHICommandListImmediate& RHICmdList, FTexture2DRHIRef RenderTargetTextureRHI)->bool
 	{
@@ -223,7 +223,7 @@ bool FExrImgMediaReaderGpu::ReadFrame(int32 FrameId, int32 MipLevel, const FImgM
 				RHICmdList.BeginRenderPass(RPInfo, TEXT("ExrTextureSwizzle"));
 
 				FExrSwizzlePS::FPermutationDomain PermutationVector;
-				PermutationVector.Set<FExrSwizzlePS::FRgbaSwizzle>(NumChannels > 3);
+				PermutationVector.Set<FExrSwizzlePS::FRgbaSwizzle>(NumChannels);
 				FExrSwizzlePS::FParameters Parameters = FExrSwizzlePS::FParameters();
 				Parameters.TextureWidth = Dim.X;
 				Parameters.TextureHeight = Dim.Y;
diff --git a/Engine/Plugins/Media/ImgMedia/Source/OpenExrWrapper/Private/OpenExrWrapper.cpp b/Engine/Plugins/Media/ImgMedia/Source/OpenExrWrapper/Private/OpenExrWrapper.cpp
--- a/Engine/Plugins/Media/ImgMedia/Source/OpenExrWrapper/Private/OpenExrWrapper.cpp
+++ b/Engine/Plugins/Media/ImgMedia/Source/OpenExrWrapper/Private/OpenExrWrapper.cpp
@@ -140,7 +140,20 @@ int32 FRgbaInputFile::GetNumChannels() const
 	int32 NumChannels = 3;
 	switch (Channels)
 	{
+	case Imf::RgbaChannels::WRITE_R:
+	case Imf::RgbaChannels::WRITE_G:
+	case Imf::RgbaChannels::WRITE_B:
+	case Imf::RgbaChannels::WRITE_A:
+	case Imf::RgbaChannels::WRITE_Y:
+	case Imf::RgbaChannels::WRITE_C:
+		NumChannels = 1;
+		break;
+	case Imf::RgbaChannels::WRITE_YC:
+	case Imf::RgbaChannels::WRITE_YA:
+		NumChannels = 2;
+		break;
 	case Imf::RgbaChannels::WRITE_RGB:
+	case Imf::RgbaChannels::WRITE_YCA:
 		NumChannels = 3;
 		break;
 	case Imf::RgbaChannels::WRITE_RGBA:
