Fixed issue where FIMGMediaScheduler::GetWorkOrReturnToPool didn't properly handle a value of 0 for the number of Cache Threads (which should use all available cores)

diff --git a/Engine/Plugins/Media/ImgMedia/Source/ImgMedia/Private/Scheduler/ImgMediaScheduler.cpp b/Engine/Plugins/Media/ImgMedia/Source/ImgMedia/Private/Scheduler/ImgMediaScheduler.cpp
--- a/Engine/Plugins/Media/ImgMedia/Source/ImgMedia/Private/Scheduler/ImgMediaScheduler.cpp
+++ b/Engine/Plugins/Media/ImgMedia/Source/ImgMedia/Private/Scheduler/ImgMediaScheduler.cpp
@@ -212,6 +212,10 @@ IQueuedWork* FImgMediaScheduler::GetWorkOrReturnToPool(FImgMediaSchedulerThread*
 #if WITH_EDITOR
 			// If there are too many workers then delete this one.
 			int32 NumWorkers = GetDefault<UImgMediaSettings>()->CacheThreads;
+			if (NumWorkers <= 0)
+			{
+				NumWorkers = FPlatformMisc::NumberOfWorkerThreadsToSpawn();
+			}
 			if (NumWorkers < AllThreads.Num())
 			{
 				AllThreads.Remove(Thread);
