ImgMedia: Loader will try to reload frames when it does not have the correct frame and is not loading any.
#rnx

diff --git a/Engine/Plugins/Media/ImgMedia/Source/ImgMedia/Private/Loader/ImgMediaLoader.cpp b/Engine/Plugins/Media/ImgMedia/Source/ImgMedia/Private/Loader/ImgMediaLoader.cpp
--- a/Engine/Plugins/Media/ImgMedia/Source/ImgMedia/Private/Loader/ImgMediaLoader.cpp
+++ b/Engine/Plugins/Media/ImgMedia/Source/ImgMedia/Private/Loader/ImgMediaLoader.cpp
@@ -70,6 +70,7 @@ FImgMediaLoader::FImgMediaLoader(const TSharedRef<FImgMediaScheduler, ESPMode::T
 	, SequenceDuration(FTimespan::Zero())
 	, SequenceFrameRate(0, 0)
 	, LastRequestedFrame(INDEX_NONE)
+	, RetryCount(0)
 	, UseGlobalCache(false)
 {
 	ResetFetchLogic();
@@ -396,6 +397,8 @@ IMediaSamples::EFetchBestSampleResult FImgMediaLoader::FetchBestVideoSampleForTi
 			// Got a potential frame?
 			if (Frame)
 			{
+				RetryCount = 0;
+
 				// Yes. First time (after flush)?
 				int32 NewSequenceIndex = QueuedSampleFetch.CurrentSequenceIndex;
 				if (QueuedSampleFetch.LastFrameIndex != INDEX_NONE)
@@ -434,6 +437,34 @@ IMediaSamples::EFetchBestSampleResult FImgMediaLoader::FetchBestVideoSampleForTi
 					}
 				}
 			}
+			else
+			{
+				// We did not get a frame...
+				// Could we have lost a frame that we previously loaded
+				// due to the global cache being full?
+				if (UseGlobalCache)
+				{
+					// Did we get a frame previously?
+					if (LastRequestedFrame != INDEX_NONE)
+					{
+						// Are we loading any frames?
+						if ((PendingFrameNumbers.Num() == 0) && (QueuedFrameNumbers.Num() == 0))
+						{
+							// Nope...
+							// Wait for this to happen for one more frame.
+							// If we have a 1 image sequence, we might have just missed the frame
+							// so try again next time.
+							RetryCount++;
+							if (RetryCount > 1)
+							{
+								UE_LOG(LogImgMedia, Error, TEXT("Reloading frames. The global cache may be too small."));
+								LastRequestedFrame = INDEX_NONE;
+								RetryCount = 0;
+							}
+						}
+					}
+				}
+			}
 		}
 	}
 	return IMediaSamples::EFetchBestSampleResult::NoSample;
diff --git a/Engine/Plugins/Media/ImgMedia/Source/ImgMedia/Private/Loader/ImgMediaLoader.h b/Engine/Plugins/Media/ImgMedia/Source/ImgMedia/Private/Loader/ImgMediaLoader.h
--- a/Engine/Plugins/Media/ImgMedia/Source/ImgMedia/Private/Loader/ImgMediaLoader.h
+++ b/Engine/Plugins/Media/ImgMedia/Source/ImgMedia/Private/Loader/ImgMediaLoader.h
@@ -425,6 +425,8 @@ private:
 
 	/** Index of the previously requested frame. */
 	int32 LastRequestedFrame;
+	/** Counter used to retry loading frames again. */
+	int32 RetryCount;
 
 	/** Collection of frame numbers that need to be read. */
 	TArray<int32> PendingFrameNumbers;
