ImgMedia: Fixed look ahead/behind getting clamped wrong and not loading enough frames to fill the cache.
#rnx

diff --git a/Engine/Plugins/Media/ImgMedia/Source/ImgMedia/Private/Loader/ImgMediaLoader.cpp b/Engine/Plugins/Media/ImgMedia/Source/ImgMedia/Private/Loader/ImgMediaLoader.cpp
--- a/Engine/Plugins/Media/ImgMedia/Source/ImgMedia/Private/Loader/ImgMediaLoader.cpp
+++ b/Engine/Plugins/Media/ImgMedia/Source/ImgMedia/Private/Loader/ImgMediaLoader.cpp
@@ -763,16 +763,16 @@ void FImgMediaLoader::LoadSequence(const FString& SequencePath, const FFrameRate
 	const SIZE_T CacheSize = FMath::Clamp(DesiredCacheSize, (SIZE_T)0, (SIZE_T)Stats.AvailablePhysical);
 
 	const int32 MaxFramesToLoad = (int32)(CacheSize / UncompressedSize);
-	const int32 NumFramesToLoad = FMath::Clamp(MaxFramesToLoad, 0, GetNumImages());
+	NumFramesToLoad = FMath::Clamp(MaxFramesToLoad, 0, GetNumImages());
 	const float LoadBehindScale = FMath::Clamp(Settings->CacheBehindPercentage, 0.0f, 100.0f) / 100.0f;
 
-	NumLoadBehind = (int32)(LoadBehindScale * NumFramesToLoad);
-	NumLoadAhead = NumFramesToLoad - NumLoadBehind;
+	NumLoadBehind = (int32)(LoadBehindScale * MaxFramesToLoad);
+	NumLoadAhead = (int32)((1.0f - LoadBehindScale) * MaxFramesToLoad);
 
 	// Giving our reader a chance to handle RAM allocation.
 	// Not all readers use this, only those that need to handle large files 
 	// or need to be as efficient as possible.
-	Reader->PreAllocateMemoryPool(NumLoadAhead + NumLoadBehind, FirstFrameInfo);
+	Reader->PreAllocateMemoryPool(NumFramesToLoad, FirstFrameInfo);
 
 	Frames.Empty(NumFramesToLoad);
 
@@ -935,7 +935,7 @@ void FImgMediaLoader::Update(int32 PlayHeadFrame, float PlayRate, bool Loop)
 	{
 		const int32 NumImagePaths = GetNumImages();
 
-		FramesToLoad.Empty(NumLoadAhead + NumLoadBehind);
+		FramesToLoad.Empty(NumFramesToLoad);
 
 		int32 FrameOffset = (PlayRate >= 0.0f) ? 1 : -1;
 
@@ -946,7 +946,9 @@ void FImgMediaLoader::Update(int32 PlayHeadFrame, float PlayRate, bool Loop)
 		int32 LoadBehindIndex = PlayHeadFrame - FrameOffset;
 
 		// alternate between look ahead and look behind
-		while ((LoadAheadCount > 0) || (LoadBehindCount > 0))
+		
+		while ((FramesToLoad.Num() < NumFramesToLoad) &&
+			((LoadAheadCount > 0) || (LoadBehindCount > 0)))
 		{
 			if (LoadAheadCount > 0)
 			{
diff --git a/Engine/Plugins/Media/ImgMedia/Source/ImgMedia/Private/Loader/ImgMediaLoader.h b/Engine/Plugins/Media/ImgMedia/Source/ImgMedia/Private/Loader/ImgMediaLoader.h
--- a/Engine/Plugins/Media/ImgMedia/Source/ImgMedia/Private/Loader/ImgMediaLoader.h
+++ b/Engine/Plugins/Media/ImgMedia/Source/ImgMedia/Private/Loader/ImgMediaLoader.h
@@ -428,6 +428,9 @@ private:
 	/** The number of frames to load behind the play head. */
 	int32 NumLoadBehind;
 
+	/** Total number of frames to load. */
+	int32 NumFramesToLoad;
+
 	/** The image sequence reader to use. */
 	TSharedPtr<IImgMediaReader, ESPMode::ThreadSafe> Reader;
 
