ImgMedia: EXR GPU reader falls back to CPU if it cannot read a file.
#rnx

diff --git a/Engine/Plugins/Media/ImgMedia/Source/ImgMedia/Private/Readers/ExrImgMediaReaderGpu.cpp b/Engine/Plugins/Media/ImgMedia/Source/ImgMedia/Private/Readers/ExrImgMediaReaderGpu.cpp
--- a/Engine/Plugins/Media/ImgMedia/Source/ImgMedia/Private/Readers/ExrImgMediaReaderGpu.cpp
+++ b/Engine/Plugins/Media/ImgMedia/Source/ImgMedia/Private/Readers/ExrImgMediaReaderGpu.cpp
@@ -124,6 +124,12 @@ FExrImgMediaReaderGpu::~FExrImgMediaReaderGpu()
 
 bool FExrImgMediaReaderGpu::ReadFrame(int32 FrameId, int32 MipLevel, const FImgMediaTileSelection& InTileSelection, TSharedPtr<FImgMediaFrame, ESPMode::ThreadSafe> OutFrame)
 {
+	// Fall back to cpu?
+	if (bFallBackToCPU)
+	{
+		return FExrImgMediaReader::ReadFrame(FrameId, MipLevel, InTileSelection, OutFrame);
+	}
+
 	TSharedPtr<FImgMediaLoader, ESPMode::ThreadSafe> Loader = LoaderPtr.Pin();
 	if (Loader.IsValid() == false)
 	{
@@ -200,7 +206,10 @@ bool FExrImgMediaReaderGpu::ReadFrame(int32 FrameId, int32 MipLevel, const FImgM
 							UE_LOG(LogImgMedia, Error, TEXT("Compressed and uncompressed files should not be mixed in a single sequence."));
 						}
 					}
-					return false;
+
+					// Fall back to CPU.
+					bFallBackToCPU = true;
+					return FExrImgMediaReader::ReadFrame(FrameId, MipLevel, InTileSelection, OutFrame);
 				}
 			}
 
diff --git a/Engine/Plugins/Media/ImgMedia/Source/ImgMedia/Private/Readers/ExrImgMediaReaderGpu.h b/Engine/Plugins/Media/ImgMedia/Source/ImgMedia/Private/Readers/ExrImgMediaReaderGpu.h
--- a/Engine/Plugins/Media/ImgMedia/Source/ImgMedia/Private/Readers/ExrImgMediaReaderGpu.h
+++ b/Engine/Plugins/Media/ImgMedia/Source/ImgMedia/Private/Readers/ExrImgMediaReaderGpu.h
@@ -49,7 +49,7 @@ public:
 
 	/** Default constructor. */
 	FExrImgMediaReaderGpu(const TSharedRef<FImgMediaLoader, ESPMode::ThreadSafe>& InLoader):FExrImgMediaReader(InLoader),
-		LastTickedFrameCounter((uint64)-1), bIsShuttingDown(false) {};
+		LastTickedFrameCounter((uint64)-1), bIsShuttingDown(false), bFallBackToCPU(false) {};
 	virtual ~FExrImgMediaReaderGpu();
 
 public:
@@ -109,6 +109,9 @@ private:
 
 	/** A flag indicating this reader is being destroyed, therefore memory should not be returned. */
 	bool bIsShuttingDown;
+
+	/** If true, then just use the CPU to read the file. */
+	bool bFallBackToCPU;
 };
 
 FUNC_DECLARE_DELEGATE(FExrConvertBufferCallback, bool, FRHICommandListImmediate& /*RHICmdList*/, FTexture2DRHIRef /*RenderTargetTextureRHI*/)
