ImgMedia: Fixed blocked playback timing out because sample did not have correct sequence index.
#rnx

diff --git a/Engine/Plugins/Media/ImgMedia/Source/ImgMedia/Private/Loader/ImgMediaLoader.cpp b/Engine/Plugins/Media/ImgMedia/Source/ImgMedia/Private/Loader/ImgMediaLoader.cpp
--- a/Engine/Plugins/Media/ImgMedia/Source/ImgMedia/Private/Loader/ImgMediaLoader.cpp
+++ b/Engine/Plugins/Media/ImgMedia/Source/ImgMedia/Private/Loader/ImgMediaLoader.cpp
@@ -396,6 +396,7 @@ IMediaSamples::EFetchBestSampleResult FImgMediaLoader::FetchBestVideoSampleForTi
 			if (Frame)
 			{
 				RetryCount = 0;
+				double Duration = Frame->Get()->Info.FrameRate.AsInterval();
 
 				// Yes. First time (after flush)?
 				int32 NewSequenceIndex = QueuedSampleFetch.CurrentSequenceIndex;
@@ -415,7 +416,32 @@ IMediaSamples::EFetchBestSampleResult FImgMediaLoader::FetchBestVideoSampleForTi
 						NewSequenceIndex = TimeRange.GetLowerBoundValue().SequenceIndex;
 					}
 				}
-
+				else
+				{
+					// Make sure this sample has the sequence index that matches the time range.
+					// If the time range only has one sequence index then just use that.
+					if (TimeRange.GetLowerBoundValue().SequenceIndex == TimeRange.GetUpperBoundValue().SequenceIndex)
+					{
+						NewSequenceIndex = TimeRange.GetLowerBoundValue().SequenceIndex;
+					}
+					else
+					{
+						// Try the lower bound sequence index.
+						NewSequenceIndex = TimeRange.GetLowerBoundValue().SequenceIndex;
+						FMediaTimeStamp SampleTime = FMediaTimeStamp(FrameNumberToTime(MaxIdx), NewSequenceIndex);
+						TRange<FMediaTimeStamp> SampleTimeRange(
+							SampleTime,
+							SampleTime + FTimespan::FromSeconds(Duration));
+						
+						// Does our sample time overlap the time range?
+						if (TimeRange.Overlaps(SampleTimeRange) == false)
+						{
+							// No. Use the upper bound sequence index.
+							NewSequenceIndex = TimeRange.GetUpperBoundValue().SequenceIndex;
+						}
+					}
+				}
+				
 				// Different from the last one we returned?
 				if ((QueuedSampleFetch.LastFrameIndex != MaxIdx) || (QueuedSampleFetch.CurrentSequenceIndex != NewSequenceIndex) || (ImagePaths.Num() == 1))
 				{
@@ -424,7 +450,7 @@ IMediaSamples::EFetchBestSampleResult FImgMediaLoader::FetchBestVideoSampleForTi
 
 					// We are clear to return it as new result... Make a sample & initialize it...
 					auto Sample = MakeShared<FImgMediaTextureSample, ESPMode::ThreadSafe>();
-					auto Duration = Frame->Get()->Info.FrameRate.AsInterval();
+					
 
 					ImgMediaLoader::CheckAndUpdateImgDimensions(SequenceDim, Frame->Get()->Info.Dim);
 
@@ -603,7 +629,7 @@ bool FImgMediaLoader::RequestFrame(FTimespan Time, float PlayRate, bool Loop)
 		// Make sure we call the reader even if we do no update - just in case it does anything
 		Reader->OnTick();
 
-		UE_LOG(LogImgMedia, VeryVerbose, TEXT("Loader %p: Skipping frame %i for time %s"), this, FrameNumber, *Time.ToString(TEXT("%h:%m:%s.%t")));
+		UE_LOG(LogImgMedia, VeryVerbose, TEXT("Loader %p: Skipping frame %i for time %s last %d"), this, FrameNumber, *Time.ToString(TEXT("%h:%m:%s.%t")), LastRequestedFrame);
 		return false;
 	}
 
