ImgMedia: Fixed StartIndex being 1 frame too late because it got rounded up when converted from StartTime.
#rnx

diff --git a/Engine/Plugins/Media/ImgMedia/Source/ImgMedia/Private/Loader/ImgMediaLoader.cpp b/Engine/Plugins/Media/ImgMedia/Source/ImgMedia/Private/Loader/ImgMediaLoader.cpp
--- a/Engine/Plugins/Media/ImgMedia/Source/ImgMedia/Private/Loader/ImgMediaLoader.cpp
+++ b/Engine/Plugins/Media/ImgMedia/Source/ImgMedia/Private/Loader/ImgMediaLoader.cpp
@@ -289,6 +289,22 @@ IMediaSamples::EFetchBestSampleResult FImgMediaLoader::FetchBestVideoSampleForTi
 		int32 MaxIdx = -1;
 		if (PlayRate >= 0.0f)
 		{
+			// Is the start index the same as the end index?
+			if (StartIndex == EndIndex)
+			{
+				// Is StartTime less than the time of StartIndex?
+				FTimespan CurrentStartTime = FrameNumberToTime(StartIndex);
+				if (StartTime < CurrentStartTime)
+				{
+					// Yes. TimeToFrameNumber rounded up, but we really dont want that
+					// so adjust StartIndex.
+					if (StartIndex > 0)
+					{
+						StartIndex--;
+					}
+				}
+			}
+
 			// Forward...
 			if (StartIndex > EndIndex)
 			{
