ImgMedia: Editor customization is now done via IPropertyTypeCustomization so it can be used with MediaBundles.
#rnx

diff --git a/Engine/Plugins/Media/ImgMedia/Source/ImgMedia/Public/ImgMediaSource.h b/Engine/Plugins/Media/ImgMedia/Source/ImgMedia/Public/ImgMediaSource.h
--- a/Engine/Plugins/Media/ImgMedia/Source/ImgMedia/Public/ImgMediaSource.h
+++ b/Engine/Plugins/Media/ImgMedia/Source/ImgMedia/Public/ImgMediaSource.h
@@ -15,6 +15,12 @@
 class AActor;
 class FImgMediaMipMapInfo;
 
+/** This provides customized editing of SequencePath. */
+USTRUCT()
+struct IMGMEDIA_API FImgMediaSourceCustomizationSequenceProxy
+{
+	GENERATED_BODY()
+};
 
 /**
  * Media source for EXR image sequences.
@@ -60,6 +66,14 @@ public:
 	UPROPERTY(BlueprintReadWrite, EditAnywhere, Category = Sequence)
 	bool bFillGapsInSequence;
 
+#if WITH_EDITORONLY_DATA
+
+	/** This is only used so we can customize editing of SequencePath. */
+	UPROPERTY(EditAnywhere, Transient, Category = Sequence)
+	FImgMediaSourceCustomizationSequenceProxy SequenceProxy;
+
+#endif // WITH_EDITORONLY_DATA
+
 public:
 
 	/**
@@ -169,7 +183,7 @@ protected:
 protected:
 
 	/** The directory that contains the image sequence files. */
-	UPROPERTY(BlueprintReadOnly, EditAnywhere, Category=Sequence)
+	UPROPERTY(BlueprintReadOnly, EditAnywhere, Category=Sequence, meta = (EditCondition = "false", EditConditionHides))
 	FDirectoryPath SequencePath;
 
 	/** MipMapInfo object to handle mip maps. */
diff --git a/Engine/Plugins/Media/ImgMedia/Source/ImgMediaEditor/Private/Customizations/ImgMediaSourceCustomization.cpp b/Engine/Plugins/Media/ImgMedia/Source/ImgMediaEditor/Private/Customizations/ImgMediaSourceCustomization.cpp
--- a/Engine/Plugins/Media/ImgMedia/Source/ImgMediaEditor/Private/Customizations/ImgMediaSourceCustomization.cpp
+++ b/Engine/Plugins/Media/ImgMedia/Source/ImgMediaEditor/Private/Customizations/ImgMediaSourceCustomization.cpp
@@ -3,6 +3,7 @@
 #include "ImgMediaSourceCustomization.h"
 
 #include "IMediaModule.h"
+#include "ImgMediaEditorModule.h"
 #include "ImgMediaMipMapInfo.h"
 #include "ImgMediaSource.h"
 #include "GameFramework/Actor.h"
@@ -22,6 +23,7 @@
 #include "DetailCategoryBuilder.h"
 #include "DetailLayoutBuilder.h"
 #include "DetailWidgetRow.h"
+#include "IDetailChildrenBuilder.h"
 #include "IDetailGroup.h"
 #include "IDetailPropertyRow.h"
 
@@ -32,89 +34,66 @@
 /* IDetailCustomization interface
  *****************************************************************************/
 
-void FImgMediaSourceCustomization::CustomizeDetails(IDetailLayoutBuilder& DetailBuilder)
+void FImgMediaSourceCustomization::CustomizeHeader(TSharedRef<IPropertyHandle> InPropertyHandle, FDetailWidgetRow& HeaderRow, IPropertyTypeCustomizationUtils& CustomizationUtils)
 {
+	PropertyHandle = InPropertyHandle;
+
 	// customize 'File' category
-	IDetailCategoryBuilder& FileCategory = DetailBuilder.EditCategory("Sequence");
-	{
-		// FilePath
-		SequencePathProperty = DetailBuilder.GetProperty("SequencePath");
-		{
-			IDetailPropertyRow& SequencePathRow = FileCategory.AddProperty(SequencePathProperty);
+	HeaderRow
+		.NameContent()
+			[
+				SNew(SHorizontalBox)
 
-			SequencePathRow
-				.ShowPropertyButtons(false)
-				.CustomWidget()
-				.NameContent()
+				+ SHorizontalBox::Slot()
+					.AutoWidth()
+					.VAlign(VAlign_Center)
 					[
-						SNew(SHorizontalBox)
-
-						+ SHorizontalBox::Slot()
-							.AutoWidth()
-							.VAlign(VAlign_Center)
-							[
-								SNew(STextBlock)
-									.Font(IDetailLayoutBuilder::GetDetailFont())
-									.Text(LOCTEXT("SequencePathPropertyName", "Sequence Path"))
-									.ToolTipText(SequencePathProperty->GetToolTipText())
-							]
-
-						+ SHorizontalBox::Slot()
-							.FillWidth(1.0f)
-							.HAlign(HAlign_Left)
-							.VAlign(VAlign_Center)
-							.Padding(4.0f, 0.0f, 0.0f, 0.0f)
-							[
-								SNew(SImage)
-									.Image(FCoreStyle::Get().GetBrush("Icons.Warning"))
-									.ToolTipText(LOCTEXT("SequencePathWarning", "The selected image sequence will not get packaged, because its path points to a directory outside the project's /Content/Movies/ directory."))
-									.Visibility(this, &FImgMediaSourceCustomization::HandleSequencePathWarningIconVisibility)
-							]
+						SNew(STextBlock)
+							.Font(IDetailLayoutBuilder::GetDetailFont())
+							.Text(LOCTEXT("SequencePathPropertyName", "Sequence Path"))
+							.ToolTipText(GetSequencePathProperty()->GetToolTipText())
 					]
-				.ValueContent()
-					.MaxDesiredWidth(0.0f)
-					.MinDesiredWidth(125.0f)
-					[
-						SNew(SFilePathPicker)
-							.BrowseButtonImage(FEditorStyle::GetBrush("PropertyWindow.Button_Ellipsis"))
-							.BrowseButtonStyle(FEditorStyle::Get(), "HoverHintOnly")
-							.BrowseButtonToolTip(LOCTEXT("SequencePathBrowseButtonToolTip", "Choose a file from this computer"))
-							.BrowseDirectory_Lambda([this]() -> FString
-							{
-								const FString SequencePath = GetSequencePath();
-								return !SequencePath.IsEmpty() ? SequencePath : (FPaths::ProjectContentDir() / TEXT("Movies"));
-							})
-							.FilePath_Lambda([this]() -> FString
-							{
-								return GetSequencePath();
-							})
-							.FileTypeFilter_Lambda([]() -> FString
-							{
-								return TEXT("All files (*.*)|*.*|EXR files (*.exr)|*.exr");
-							})
-							.OnPathPicked(this, &FImgMediaSourceCustomization::HandleSequencePathPickerPathPicked)
-							.ToolTipText(LOCTEXT("SequencePathToolTip", "The path to an image sequence file on this computer"))
-					];
-		}
-
-		PathRelativeToRootProperty = DetailBuilder.GetProperty("IsPathRelativeToProjectRoot");
-	}
-
-	// add 'Proxies' category
-	IDetailCategoryBuilder& ProxiesCategory = DetailBuilder.EditCategory("Proxies", LOCTEXT("ProxiesCategoryName", "Proxies"));
-	{
-		// add 'Proxy Directories' row
-		FDetailWidgetRow& PreviewRow = ProxiesCategory.AddCustomRow(LOCTEXT("ProxiesRowFilterString", "Proxy Directories"));
 
-		PreviewRow.WholeRowContent()
+				+ SHorizontalBox::Slot()
+					.FillWidth(1.0f)
+					.HAlign(HAlign_Left)
+					.VAlign(VAlign_Center)
+					.Padding(4.0f, 0.0f, 0.0f, 0.0f)
+					[
+						SNew(SImage)
+							.Image(FCoreStyle::Get().GetBrush("Icons.Warning"))
+							.ToolTipText(LOCTEXT("SequencePathWarning", "The selected image sequence will not get packaged, because its path points to a directory outside the project's /Content/Movies/ directory."))
+							.Visibility(this, &FImgMediaSourceCustomization::HandleSequencePathWarningIconVisibility)
+					]
+			]
+		.ValueContent()
+			.MaxDesiredWidth(0.0f)
+			.MinDesiredWidth(125.0f)
 			[
-				SAssignNew(ProxiesTextBlock, SEditableTextBox)
-					.IsReadOnly(true)
+				SNew(SFilePathPicker)
+					.BrowseButtonImage(FEditorStyle::GetBrush("PropertyWindow.Button_Ellipsis"))
+					.BrowseButtonStyle(FEditorStyle::Get(), "HoverHintOnly")
+					.BrowseButtonToolTip(LOCTEXT("SequencePathBrowseButtonToolTip", "Choose a file from this computer"))
+					.BrowseDirectory_Lambda([this]() -> FString
+					{
+						const FString SequencePath = GetSequencePath();
+						return !SequencePath.IsEmpty() ? SequencePath : (FPaths::ProjectContentDir() / TEXT("Movies"));
+					})
+					.FilePath_Lambda([this]() -> FString
+					{
+						return GetSequencePath();
+					})
+					.FileTypeFilter_Lambda([]() -> FString
+					{
+						return TEXT("All files (*.*)|*.*|EXR files (*.exr)|*.exr");
+					})
+					.OnPathPicked(this, &FImgMediaSourceCustomization::HandleSequencePathPickerPathPicked)
+					.ToolTipText(LOCTEXT("SequencePathToolTip", "The path to an image sequence file on this computer"))
 			];
-	}
+}
 
-	// Add mipmap info.
-	CustomizeMipMapInfo(DetailBuilder);
+void FImgMediaSourceCustomization::CustomizeChildren(TSharedRef<IPropertyHandle> InStructPropertyHandle, IDetailChildrenBuilder& StructBuilder, IPropertyTypeCustomizationUtils& StructCustomizationUtils)
+{
 }
 
 void FImgMediaSourceCustomization::CustomizeMipMapInfo(IDetailLayoutBuilder& DetailBuilder)
@@ -310,7 +289,14 @@ void FImgMediaSourceCustomization::AddCameraMipDistances(IDetailGroup& InCameraG
 FString FImgMediaSourceCustomization::GetSequencePath() const
 {
 	FString FilePath;
-	SequencePathProperty->GetChildHandle("Path")->GetValue(FilePath);
+	TSharedPtr<IPropertyHandle> SequencePathProperty = GetSequencePathPathProperty();
+	if (SequencePathProperty.IsValid())
+	{
+		if (SequencePathProperty->GetValue(FilePath) != FPropertyAccess::Success)
+		{
+			UE_LOG(LogImgMediaEditor, Error, TEXT("FImgMediaSourceCustomization could not get SequencePath."));
+		}
+	}
 
 	return FilePath;
 }
@@ -318,8 +304,8 @@ FString FImgMediaSourceCustomization::GetSequencePath() const
 FString FImgMediaSourceCustomization::GetRelativePathRoot() const
 {
 	FString RelativeDir;
-	bool bIsPathRelativeToRoot = false;
-	PathRelativeToRootProperty->GetValue(bIsPathRelativeToRoot);
+	bool bIsPathRelativeToRoot = IsPathRelativeToRoot();
+
 	if (bIsPathRelativeToRoot)
 	{
 		RelativeDir = FPaths::ConvertRelativePathToFull(FPaths::ProjectDir());
@@ -331,6 +317,83 @@ FString FImgMediaSourceCustomization::GetRelativePathRoot() const
 	return RelativeDir;
 }
 
+
+TSharedPtr<IPropertyHandle> FImgMediaSourceCustomization::GetSequencePathProperty() const
+{
+	TSharedPtr<IPropertyHandle> SequencePathProperty;
+
+	if ((PropertyHandle.IsValid()) && (PropertyHandle->IsValidHandle()))
+	{
+		TSharedPtr<IPropertyHandle> ParentHandle = PropertyHandle->GetParentHandle();
+		if (ParentHandle.IsValid())
+		{
+			SequencePathProperty = ParentHandle->GetChildHandle("SequencePath");
+		}
+	}
+
+	return SequencePathProperty;
+}
+
+TSharedPtr<IPropertyHandle> FImgMediaSourceCustomization::GetSequencePathPathProperty() const
+{
+	TSharedPtr<IPropertyHandle> SequencePathPathProperty;
+
+	TSharedPtr<IPropertyHandle> SequencePathProperty = GetSequencePathProperty();
+	if (SequencePathProperty.IsValid())
+	{
+		SequencePathPathProperty = SequencePathProperty->GetChildHandle("Path");
+	}
+
+	return SequencePathPathProperty;
+}
+
+
+TSharedPtr<IPropertyHandle> FImgMediaSourceCustomization::GetPathRelativeToRootProperty() const
+{
+	TSharedPtr<IPropertyHandle> PathRelativeToRootProperty;
+
+	if ((PropertyHandle.IsValid()) && (PropertyHandle->IsValidHandle()))
+	{
+		TSharedPtr<IPropertyHandle> ParentHandle = PropertyHandle->GetParentHandle();
+		if (ParentHandle.IsValid())
+		{
+			PathRelativeToRootProperty = ParentHandle->GetChildHandle("IsPathRelativeToProjectRoot");
+		}
+	}
+
+	return PathRelativeToRootProperty;
+}
+
+
+bool FImgMediaSourceCustomization::IsPathRelativeToRoot() const
+{
+	TSharedPtr<IPropertyHandle> PathRelativeToRootProperty = GetPathRelativeToRootProperty();
+
+	bool bIsPathRelativeToRoot = false;
+	if (PathRelativeToRootProperty.IsValid())
+	{
+		if (PathRelativeToRootProperty->GetValue(bIsPathRelativeToRoot) != FPropertyAccess::Success)
+		{
+			UE_LOG(LogImgMediaEditor, Error, TEXT("FImgMediaSourceCustomization could not get IsPathRelativeToProjectRoot."));
+		}
+	}
+
+	return bIsPathRelativeToRoot;
+}
+
+void FImgMediaSourceCustomization::SetPathRelativeToRoot(bool bIsPathRelativeToRoot)
+{
+	TSharedPtr<IPropertyHandle> PathRelativeToRootProperty = GetPathRelativeToRootProperty();
+
+	if (PathRelativeToRootProperty.IsValid())
+	{
+		if (PathRelativeToRootProperty->SetValue(bIsPathRelativeToRoot) != FPropertyAccess::Success)
+		{
+			UE_LOG(LogImgMediaEditor, Error, TEXT("FImgMediaSourceCustomization could not set IsPathRelativeToProjectRoot."));
+		}
+	}
+}
+
 /* FImgMediaSourceCustomization callbacks
  *****************************************************************************/
 
@@ -355,19 +418,25 @@ void FImgMediaSourceCustomization::HandleSequencePathPickerPathPicked(const FStr
 	}
 	
 	// Update relative to root property.
-	PathRelativeToRootProperty->SetValue(bIsRelativePath);
+	SetPathRelativeToRoot(bIsRelativePath);
 
 	// update property
-	TSharedPtr<IPropertyHandle> SequencePathPathProperty = SequencePathProperty->GetChildHandle("Path");
-	SequencePathPathProperty->SetValue(PickedDir);
+	TSharedPtr<IPropertyHandle> SequencePathPathProperty = GetSequencePathPathProperty();
+	if (SequencePathPathProperty.IsValid())
+	{
+		if (SequencePathPathProperty->SetValue(PickedDir) != FPropertyAccess::Success)
+		{
+			UE_LOG(LogImgMediaEditor, Error, TEXT("FImgMediaSourceCustomization could not set SequencePath."));
+		}
+	}
 }
 
 
 EVisibility FImgMediaSourceCustomization::HandleSequencePathWarningIconVisibility() const
 {
-	FString FilePath;
+	FString FilePath = GetSequencePath();
 
-	if ((SequencePathProperty->GetChildHandle("Path")->GetValue(FilePath) != FPropertyAccess::Success) || FilePath.IsEmpty() || FilePath.Contains(TEXT("://")))
+	if (FilePath.IsEmpty() || FilePath.Contains(TEXT("://")))
 	{
 		return EVisibility::Hidden;
 	}
diff --git a/Engine/Plugins/Media/ImgMedia/Source/ImgMediaEditor/Private/Customizations/ImgMediaSourceCustomization.h b/Engine/Plugins/Media/ImgMedia/Source/ImgMediaEditor/Private/Customizations/ImgMediaSourceCustomization.h
--- a/Engine/Plugins/Media/ImgMedia/Source/ImgMediaEditor/Private/Customizations/ImgMediaSourceCustomization.h
+++ b/Engine/Plugins/Media/ImgMedia/Source/ImgMediaEditor/Private/Customizations/ImgMediaSourceCustomization.h
@@ -4,8 +4,8 @@
 
 #include "CoreTypes.h"
 #include "Containers/UnrealString.h"
+#include "IPropertyTypeCustomization.h"
 #include "Layout/Visibility.h"
-#include "IDetailCustomization.h"
 #include "Templates/SharedPointer.h"
 
 class IDetailCategoryBuilder;
@@ -22,18 +22,13 @@ struct FImgMediaMipMapObjectInfo;
  * Implements a details view customization for the UImgMediaSource class.
  */
 class FImgMediaSourceCustomization
-	: public IDetailCustomization
+	: public IPropertyTypeCustomization
 {
 public:
 
-	/** Virtual destructor. */
-	~FImgMediaSourceCustomization() { }
-
-public:
-
-	//~ IDetailCustomization interface
-
-	virtual void CustomizeDetails(IDetailLayoutBuilder& DetailBuilder) override;
+	/** IPropertyTypeCustomization interface */
+	virtual void CustomizeHeader(TSharedRef<class IPropertyHandle> StructPropertyHandle, class FDetailWidgetRow& HeaderRow, IPropertyTypeCustomizationUtils& StructCustomizationUtils) override;
+	virtual void CustomizeChildren(TSharedRef<class IPropertyHandle> StructPropertyHandle, class IDetailChildrenBuilder& StructBuilder, IPropertyTypeCustomizationUtils& StructCustomizationUtils) override;
 
 public:
 
@@ -42,7 +37,7 @@ public:
 	 *
 	 * @return The new instance.
 	 */
-	static TSharedRef<IDetailCustomization> MakeInstance()
+	static TSharedRef<IPropertyTypeCustomization> MakeInstance()
 	{
 		return MakeShareable(new FImgMediaSourceCustomization());
 	}
@@ -105,13 +100,22 @@ private:
 	/** Callback for getting the visibility of warning icon for invalid SequencePath paths. */
 	EVisibility HandleSequencePathWarningIconVisibility() const;
 
+	/** Returns the property for SequencePath. */
+	TSharedPtr<IPropertyHandle> GetSequencePathProperty() const;
+	/** Returns the property for SequencePath->Path. */
+	TSharedPtr<IPropertyHandle> GetSequencePathPathProperty() const;
+	/** Returns the property for IsPathRelativeToProjectRoot. */
+	TSharedPtr<IPropertyHandle> GetPathRelativeToRootProperty() const;
+	/** Returns the value of IsPathRelativeToProjectRoot. */
+	bool IsPathRelativeToRoot() const;
+	/** Sets the value of IsPathRelativeToProjectRoot. */
+	void SetPathRelativeToRoot(bool bIsPathRelativeToRoot);
+
 private:
 
 	/** Text block widget showing the found proxy directories. */
 	TSharedPtr<SEditableTextBox> ProxiesTextBlock;
 
-	/** Pointer to the SequencePath.Path property handle. */
-	TSharedPtr<IPropertyHandle> SequencePathProperty;
-	/** Pointer to the IsPathRelativeToProjectRoot property handle. */
-	TSharedPtr<IPropertyHandle> PathRelativeToRootProperty;
+	/** Stores our property. */
+	TSharedPtr<IPropertyHandle> PropertyHandle;
 };
diff --git a/Engine/Plugins/Media/ImgMedia/Source/ImgMediaEditor/Private/ImgMediaEditorModule.cpp b/Engine/Plugins/Media/ImgMedia/Source/ImgMediaEditor/Private/ImgMediaEditorModule.cpp
--- a/Engine/Plugins/Media/ImgMedia/Source/ImgMediaEditor/Private/ImgMediaEditorModule.cpp
+++ b/Engine/Plugins/Media/ImgMedia/Source/ImgMediaEditor/Private/ImgMediaEditorModule.cpp
@@ -45,7 +45,9 @@ protected:
 
 		FPropertyEditorModule& PropertyModule = FModuleManager::LoadModuleChecked<FPropertyEditorModule>("PropertyEditor");
 		{
-			PropertyModule.RegisterCustomClassLayout(ImgMediaSourceName, FOnGetDetailCustomizationInstance::CreateStatic(&FImgMediaSourceCustomization::MakeInstance));
+#if WITH_EDITORONLY_DATA
+			PropertyModule.RegisterCustomPropertyTypeLayout(FImgMediaSourceCustomizationSequenceProxy::StaticStruct()->GetFName(), FOnGetPropertyTypeCustomizationInstance::CreateStatic(&FImgMediaSourceCustomization::MakeInstance));
+#endif // WITH_EDITORONLY_DATA
 		}
 	}
 
@@ -54,7 +56,9 @@ protected:
 	{
 		FPropertyEditorModule& PropertyModule = FModuleManager::LoadModuleChecked<FPropertyEditorModule>("PropertyEditor");
 		{
-			PropertyModule.UnregisterCustomClassLayout(ImgMediaSourceName);
+#if WITH_EDITORONLY_DATA
+			PropertyModule.UnregisterCustomPropertyTypeLayout(FImgMediaSourceCustomizationSequenceProxy::StaticStruct()->GetFName());
+#endif // WITH_EDITORONLY_DATA
 		}
 	}
 
