ImgMedia: Fix for crash when looping a video of length one image.
#rnx

diff --git a/Engine/Plugins/Media/ImgMedia/Source/ImgMedia/Private/Loader/ImgMediaLoader.cpp b/Engine/Plugins/Media/ImgMedia/Source/ImgMedia/Private/Loader/ImgMediaLoader.cpp
--- a/Engine/Plugins/Media/ImgMedia/Source/ImgMedia/Private/Loader/ImgMediaLoader.cpp
+++ b/Engine/Plugins/Media/ImgMedia/Source/ImgMedia/Private/Loader/ImgMediaLoader.cpp
@@ -256,7 +256,9 @@ IMediaSamples::EFetchBestSampleResult FImgMediaLoader::FetchBestVideoSampleForTi
 	if (IsInitialized() && TimeRange.HasLowerBound() && TimeRange.HasUpperBound())
 	{
 		FTimespan StartTime = TimeRange.GetLowerBoundValue().Time;
+		int64 StartSequenceIndex = TimeRange.GetLowerBoundValue().SequenceIndex;
 		FTimespan EndTime = TimeRange.GetUpperBoundValue().Time;
+		int64 EndSequenceIndex = TimeRange.GetUpperBoundValue().SequenceIndex;
 		check(TimeRange.GetLowerBoundValue() <= TimeRange.GetUpperBoundValue());
 
 		if (bIsLoopingEnabled)
@@ -306,7 +308,7 @@ IMediaSamples::EFetchBestSampleResult FImgMediaLoader::FetchBestVideoSampleForTi
 			}
 
 			// Forward...
-			if (StartIndex > EndIndex)
+			if ((StartIndex > EndIndex) || (EndSequenceIndex > StartSequenceIndex))
 			{
 				int32 MaxIdx1, MaxIdx2;
 				float MaxOverlap1 = FindMaxOverlapInRange(StartIndex, GetNumImages()-1, StartTime, FrameNumberToTime(GetNumImages()), MaxIdx1);
@@ -321,7 +323,7 @@ IMediaSamples::EFetchBestSampleResult FImgMediaLoader::FetchBestVideoSampleForTi
 		else
 		{
 			// Backward...
-			if (StartIndex > EndIndex)
+			if ((StartIndex > EndIndex) || (StartSequenceIndex < EndSequenceIndex))
 			{
 				int32 MaxIdx1, MaxIdx2;
 				float MaxOverlap1 = FindMaxOverlapInRange(EndIndex, 0, FTimespan::Zero(), EndTime, MaxIdx1);
