ImgMedia: Exposed function to get the sequence path in the customization.
#rnx

diff --git a/Engine/Plugins/Media/ImgMedia/Source/ImgMediaEditor/Private/Customizations/ImgMediaSourceCustomization.cpp b/Engine/Plugins/Media/ImgMedia/Source/ImgMediaEditor/Private/Customizations/ImgMediaSourceCustomization.cpp
--- a/Engine/Plugins/Media/ImgMedia/Source/ImgMediaEditor/Private/Customizations/ImgMediaSourceCustomization.cpp
+++ b/Engine/Plugins/Media/ImgMedia/Source/ImgMediaEditor/Private/Customizations/ImgMediaSourceCustomization.cpp
@@ -51,7 +51,7 @@ void FImgMediaSourceCustomization::CustomizeHeader(TSharedRef<IPropertyHandle> I
 						SNew(STextBlock)
 							.Font(IDetailLayoutBuilder::GetDetailFont())
 							.Text(LOCTEXT("SequencePathPropertyName", "Sequence Path"))
-							.ToolTipText(GetSequencePathProperty()->GetToolTipText())
+							.ToolTipText(GetSequencePathProperty(PropertyHandle)->GetToolTipText())
 					]
 
 				+ SHorizontalBox::Slot()
@@ -286,10 +286,10 @@ void FImgMediaSourceCustomization::AddCameraMipDistances(IDetailGroup& InCameraG
 /* FImgMediaSourceCustomization implementation
  *****************************************************************************/
 
-FString FImgMediaSourceCustomization::GetSequencePath() const
+FString FImgMediaSourceCustomization::GetSequencePathFromChildProperty(const TSharedPtr<IPropertyHandle>& InPropertyHandle)
 {
 	FString FilePath;
-	TSharedPtr<IPropertyHandle> SequencePathProperty = GetSequencePathPathProperty();
+	TSharedPtr<IPropertyHandle> SequencePathProperty = GetSequencePathPathProperty(InPropertyHandle);
 	if (SequencePathProperty.IsValid())
 	{
 		if (SequencePathProperty->GetValue(FilePath) != FPropertyAccess::Success)
@@ -301,6 +301,13 @@ FString FImgMediaSourceCustomization::GetSequencePath() const
 	return FilePath;
 }
 
+FString FImgMediaSourceCustomization::GetSequencePath() const
+{
+	FString FilePath = GetSequencePathFromChildProperty(PropertyHandle);
+
+	return FilePath;
+}
+
 FString FImgMediaSourceCustomization::GetRelativePathRoot() const
 {
 	FString RelativeDir;
@@ -318,13 +325,13 @@ FString FImgMediaSourceCustomization::GetRelativePathRoot() const
 }
 
 
-TSharedPtr<IPropertyHandle> FImgMediaSourceCustomization::GetSequencePathProperty() const
+TSharedPtr<IPropertyHandle> FImgMediaSourceCustomization::GetSequencePathProperty(const TSharedPtr<IPropertyHandle>& InPropertyHandle)
 {
 	TSharedPtr<IPropertyHandle> SequencePathProperty;
 
-	if ((PropertyHandle.IsValid()) && (PropertyHandle->IsValidHandle()))
+	if ((InPropertyHandle.IsValid()) && (InPropertyHandle->IsValidHandle()))
 	{
-		TSharedPtr<IPropertyHandle> ParentHandle = PropertyHandle->GetParentHandle();
+		TSharedPtr<IPropertyHandle> ParentHandle = InPropertyHandle->GetParentHandle();
 		if (ParentHandle.IsValid())
 		{
 			SequencePathProperty = ParentHandle->GetChildHandle("SequencePath");
@@ -334,11 +341,11 @@ TSharedPtr<IPropertyHandle> FImgMediaSourceCustomization::GetSequencePathPropert
 	return SequencePathProperty;
 }
 
-TSharedPtr<IPropertyHandle> FImgMediaSourceCustomization::GetSequencePathPathProperty() const
+TSharedPtr<IPropertyHandle> FImgMediaSourceCustomization::GetSequencePathPathProperty(const TSharedPtr<IPropertyHandle>& InPropertyHandle)
 {
 	TSharedPtr<IPropertyHandle> SequencePathPathProperty;
 
-	TSharedPtr<IPropertyHandle> SequencePathProperty = GetSequencePathProperty();
+	TSharedPtr<IPropertyHandle> SequencePathProperty = GetSequencePathProperty(InPropertyHandle);
 	if (SequencePathProperty.IsValid())
 	{
 		SequencePathPathProperty = SequencePathProperty->GetChildHandle("Path");
@@ -421,7 +428,7 @@ void FImgMediaSourceCustomization::HandleSequencePathPickerPathPicked(const FStr
 	SetPathRelativeToRoot(bIsRelativePath);
 
 	// update property
-	TSharedPtr<IPropertyHandle> SequencePathPathProperty = GetSequencePathPathProperty();
+	TSharedPtr<IPropertyHandle> SequencePathPathProperty = GetSequencePathPathProperty(PropertyHandle);
 	if (SequencePathPathProperty.IsValid())
 	{
 		if (SequencePathPathProperty->SetValue(PickedDir) != FPropertyAccess::Success)
diff --git a/Engine/Plugins/Media/ImgMedia/Source/ImgMediaEditor/Private/Customizations/ImgMediaSourceCustomization.h b/Engine/Plugins/Media/ImgMedia/Source/ImgMediaEditor/Private/Customizations/ImgMediaSourceCustomization.h
--- a/Engine/Plugins/Media/ImgMedia/Source/ImgMediaEditor/Private/Customizations/ImgMediaSourceCustomization.h
+++ b/Engine/Plugins/Media/ImgMedia/Source/ImgMediaEditor/Private/Customizations/ImgMediaSourceCustomization.h
@@ -42,6 +42,15 @@ public:
 		return MakeShareable(new FImgMediaSourceCustomization());
 	}
 
+	/**
+	 * Get the path to the currently selected image sequence.
+	 *
+	 * @param InPropertyHandle A property that is a child of the top level ImgMediaSource property.
+	 * 
+	 * @return Sequence path string.
+	 */
+	static FString GetSequencePathFromChildProperty(const TSharedPtr<IPropertyHandle>& InPropertyHandle);
+
 protected:
 
 	/**
@@ -101,9 +110,9 @@ private:
 	EVisibility HandleSequencePathWarningIconVisibility() const;
 
 	/** Returns the property for SequencePath. */
-	TSharedPtr<IPropertyHandle> GetSequencePathProperty() const;
+	static TSharedPtr<IPropertyHandle> GetSequencePathProperty(const TSharedPtr<IPropertyHandle>& InPropertyHandle);
 	/** Returns the property for SequencePath->Path. */
-	TSharedPtr<IPropertyHandle> GetSequencePathPathProperty() const;
+	static TSharedPtr<IPropertyHandle> GetSequencePathPathProperty(const TSharedPtr<IPropertyHandle>& InPropertyHandle);
 	/** Returns the property for IsPathRelativeToProjectRoot. */
 	TSharedPtr<IPropertyHandle> GetPathRelativeToRootProperty() const;
 	/** Returns the value of IsPathRelativeToProjectRoot. */
