- Fixing imgmediaeditor shutdown

diff --git a/Engine/Plugins/Media/ImgMedia/Source/ImgMediaEditor/Private/ImgMediaEditorModule.cpp b/Engine/Plugins/Media/ImgMedia/Source/ImgMediaEditor/Private/ImgMediaEditorModule.cpp
--- a/Engine/Plugins/Media/ImgMedia/Source/ImgMediaEditor/Private/ImgMediaEditorModule.cpp
+++ b/Engine/Plugins/Media/ImgMedia/Source/ImgMediaEditor/Private/ImgMediaEditorModule.cpp
@@ -41,12 +41,12 @@ protected:
 	/** Register details view customizations. */
 	void RegisterCustomizations()
 	{
-		ImgMediaSourceName = UImgMediaSource::StaticClass()->GetFName();
+		CustomizedStructName = FImgMediaSourceCustomizationSequenceProxy::StaticStruct()->GetFName();
 
 		FPropertyEditorModule& PropertyModule = FModuleManager::LoadModuleChecked<FPropertyEditorModule>("PropertyEditor");
 		{
 #if WITH_EDITORONLY_DATA
-			PropertyModule.RegisterCustomPropertyTypeLayout(FImgMediaSourceCustomizationSequenceProxy::StaticStruct()->GetFName(), FOnGetPropertyTypeCustomizationInstance::CreateStatic(&FImgMediaSourceCustomization::MakeInstance));
+			PropertyModule.RegisterCustomPropertyTypeLayout(CustomizedStructName, FOnGetPropertyTypeCustomizationInstance::CreateStatic(&FImgMediaSourceCustomization::MakeInstance));
 #endif // WITH_EDITORONLY_DATA
 		}
 	}
@@ -57,7 +57,7 @@ protected:
 		FPropertyEditorModule& PropertyModule = FModuleManager::LoadModuleChecked<FPropertyEditorModule>("PropertyEditor");
 		{
 #if WITH_EDITORONLY_DATA
-			PropertyModule.UnregisterCustomPropertyTypeLayout(FImgMediaSourceCustomizationSequenceProxy::StaticStruct()->GetFName());
+			PropertyModule.UnregisterCustomPropertyTypeLayout(CustomizedStructName);
 #endif // WITH_EDITORONLY_DATA
 		}
 	}
@@ -88,8 +88,8 @@ protected:
 
 private:
 
-	/** Class names. */
-	FName ImgMediaSourceName;
+	/** Customization name to avoid reusing staticstruct during shutdown. */
+	FName CustomizedStructName;
 
 	/** The collection of registered asset type actions. */
 	TArray<TSharedRef<IAssetTypeActions>> RegisteredAssetTypeActions;
