ImgMedia: Added framework for importing images.

#rnx

diff --git a/Engine/Plugins/Media/ImgMedia/Source/ImgMedia/Public/ImgMediaSource.h b/Engine/Plugins/Media/ImgMedia/Source/ImgMedia/Public/ImgMediaSource.h
--- a/Engine/Plugins/Media/ImgMedia/Source/ImgMedia/Public/ImgMediaSource.h
+++ b/Engine/Plugins/Media/ImgMedia/Source/ImgMedia/Public/ImgMediaSource.h
@@ -22,6 +22,13 @@ struct IMGMEDIA_API FImgMediaSourceCustomizationSequenceProxy
 	GENERATED_BODY()
 };
 
+/** Information on the images that have been imported. */
+USTRUCT(BlueprintType)
+struct IMGMEDIA_API FImgMediaSourceImportInfo
+{
+	GENERATED_BODY()
+};
+
 /**
  * Media source for EXR image sequences.
  *
@@ -74,6 +81,10 @@ public:
 
 #endif // WITH_EDITORONLY_DATA
 
+	/** Info concerning the imported images. */
+	UPROPERTY(BlueprintReadWrite, EditAnywhere, Category = Sequence)
+	FImgMediaSourceImportInfo ImportInfo;
+
 public:
 
 	/**
diff --git a/Engine/Plugins/Media/ImgMedia/Source/ImgMediaEditor/ImgMediaEditor.Build.cs b/Engine/Plugins/Media/ImgMedia/Source/ImgMediaEditor/ImgMediaEditor.Build.cs
--- a/Engine/Plugins/Media/ImgMedia/Source/ImgMediaEditor/ImgMediaEditor.Build.cs
+++ b/Engine/Plugins/Media/ImgMedia/Source/ImgMediaEditor/ImgMediaEditor.Build.cs
@@ -14,7 +14,9 @@ namespace UnrealBuildTool.Rules
 					"DesktopWidgets",
 					"EditorFramework",
 					"EditorStyle",
+					"ImageWrapper",
 					"ImgMedia",
+					"InputCore",
 					"MediaAssets",
 					"Slate",
 					"SlateCore",
diff --git a/Engine/Plugins/Media/ImgMedia/Source/ImgMediaEditor/Private/Customizations/ImgMediaSourceCustomizationImportInfo.cpp b/Engine/Plugins/Media/ImgMedia/Source/ImgMediaEditor/Private/Customizations/ImgMediaSourceCustomizationImportInfo.cpp
new file mode 100644
--- /dev/null
+++ b/Engine/Plugins/Media/ImgMedia/Source/ImgMediaEditor/Private/Customizations/ImgMediaSourceCustomizationImportInfo.cpp
@@ -0,0 +1,212 @@
+// Copyright Epic Games, Inc. All Rights Reserved.
+
+#include "ImgMediaSourceCustomizationImportInfo.h"
+
+#include "Async/Async.h"
+#include "DetailLayoutBuilder.h"
+#include "DetailWidgetRow.h"
+#include "EditorStyleSet.h"
+#include "Framework/Notifications/NotificationManager.h"
+#include "HAL/PlatformFileManager.h"
+#include "IDetailChildrenBuilder.h"
+#include "IDetailGroup.h"
+#include "IImageWrapperModule.h"
+#include "ImageWrapperHelper.h"
+#include "IMediaModule.h"
+#include "ImgMediaSourceCustomization.h"
+#include "ImgMediaEditorModule.h"
+#include "ImgMediaSource.h"
+#include "Misc/FileHelper.h"
+#include "Misc/Paths.h"
+#include "Modules/ModuleManager.h"
+#include "Styling/CoreStyle.h"
+#include "Widgets/DeclarativeSyntaxSupport.h"
+#include "Widgets/Images/SImage.h"
+#include "Widgets/Input/SButton.h"
+#include "Widgets/Input/SEditableTextBox.h"
+#include "Widgets/Input/SFilePathPicker.h"
+#include "Widgets/Input/SNumericEntryBox.h"
+#include "Widgets/Notifications/SNotificationList.h"
+#include "Widgets/SBoxPanel.h"
+#include "Widgets/Text/STextBlock.h"
+
+#define LOCTEXT_NAMESPACE "FImgMediaSourceCustomizationImportInfo"
+
+/* IPropertyTypeCustomization interface
+ *****************************************************************************/
+
+void FImgMediaSourceCustomizationImportInfo::CustomizeHeader(TSharedRef<IPropertyHandle> InPropertyHandle, FDetailWidgetRow& HeaderRow, IPropertyTypeCustomizationUtils& CustomizationUtils)
+{
+	PropertyHandle = InPropertyHandle;
+
+	HeaderRow
+		.NameContent()
+		[
+			SNew(STextBlock)
+				.Text(LOCTEXT("SequencePathImport", "Import Sequence"))
+				.Font(IDetailLayoutBuilder::GetDetailFont())
+		]
+		.ValueContent()
+		[
+			SNew(SHorizontalBox)
+
+			+ SHorizontalBox::Slot()
+				.AutoWidth()
+				.VAlign(VAlign_Center)
+				[
+					SNew(SButton)
+						.ButtonStyle(FEditorStyle::Get(), "HoverHintOnly")
+						.ForegroundColor(FSlateColor::UseForeground())
+						.OnClicked(this, &FImgMediaSourceCustomizationImportInfo::OnImportClicked)
+						.ToolTipText(LOCTEXT("SequencePathImportButtonToolTip", "Import files for this sequence and generate tiles and/or mips if requested."))
+						[
+							SNew(SImage)
+								.ColorAndOpacity(FSlateColor::UseForeground())
+								.Image(FEditorStyle::GetBrush("Icons.Adjust"))
+						]
+				]
+		];
+}
+
+void FImgMediaSourceCustomizationImportInfo::CustomizeChildren(TSharedRef<IPropertyHandle> InStructPropertyHandle, IDetailChildrenBuilder& StructBuilder, IPropertyTypeCustomizationUtils& StructCustomizationUtils)
+{
+	IDetailGroup& TileGroup = StructBuilder.AddGroup(FName(TEXT("TilesGroup")), LOCTEXT("TileGroup_DisplayName", "Tiles"));
+	TileGroup.AddWidgetRow()
+		.NameContent()
+		[
+			SNew(STextBlock)
+				.Text(LOCTEXT("ImportTileWidth", "Width"))
+				.Font(IDetailLayoutBuilder::GetDetailFont())
+		]
+		.ValueContent()
+		[
+			SNew(SNumericEntryBox<int32>)
+				.Value(this, &FImgMediaSourceCustomizationImportInfo::GetTileWidth)
+				.OnValueChanged(this, &FImgMediaSourceCustomizationImportInfo::SetTileWidth)
+		];
+}
+
+/* FImgMediaSourceCustomizationImport implementation
+ *****************************************************************************/
+
+TOptional<int32> FImgMediaSourceCustomizationImportInfo::GetTileWidth() const
+{
+	return TileWidth;
+}
+
+void FImgMediaSourceCustomizationImportInfo::SetTileWidth(int32 InWidth)
+{
+	TileWidth = InWidth;
+}
+
+FReply FImgMediaSourceCustomizationImportInfo::OnImportClicked()
+{
+	// Create notification.
+	FNotificationInfo Info(FText::GetEmpty());
+	Info.bFireAndForget = false;
+	TSharedPtr<SNotificationItem> ConfirmNotification = FSlateNotificationManager::Get().AddNotification(Info);
+
+	// Start async task to import files.
+	Async(EAsyncExecution::Thread, [this, ConfirmNotification]()
+	{
+		FString SequencePath = FImgMediaSourceCustomization::GetSequencePathFromChildProperty(PropertyHandle);
+		ImportFiles(SequencePath, ConfirmNotification);
+	});
+
+	return FReply::Handled();
+}
+
+void FImgMediaSourceCustomizationImportInfo::ImportFiles(const FString& SequencePath, TSharedPtr<SNotificationItem> ConfirmNotification)
+{
+	// Create output directory.
+	FString OutPath = FPaths::Combine(SequencePath, TEXT("Imported"));
+	IPlatformFile& PlatformFile = FPlatformFileManager::Get().GetPlatformFile();
+	PlatformFile.CreateDirectoryTree(*OutPath);
+
+	// Get source files.
+	TArray<FString> FoundFiles;
+	IFileManager::Get().FindFiles(FoundFiles, *SequencePath, TEXT("*"));
+	FoundFiles.Sort();
+	UE_LOG(LogImgMediaEditor, Warning, TEXT("Found %i image files in %s to import."), FoundFiles.Num(), *SequencePath);
+	if (FoundFiles.Num() == 0)
+	{
+		UE_LOG(LogImgMediaEditor, Error, TEXT("No files to import."));
+	}
+	else
+	{
+		// Create image wrapper
+		FString Ext = FPaths::GetExtension(FoundFiles[0]);
+		EImageFormat ImageFormat = ImageWrapperHelper::GetImageFormat(Ext);
+
+		if (ImageFormat == EImageFormat::Invalid)
+		{
+			UE_LOG(LogImgMediaEditor, Error, TEXT("Invalid file format %s"), *Ext);
+		}
+		else
+		{
+			IImageWrapperModule& ImageWrapperModule = FModuleManager::LoadModuleChecked<IImageWrapperModule>("ImageWrapper");
+			TSharedPtr<IImageWrapper> ImageWrapper = ImageWrapperModule.CreateImageWrapper(ImageFormat);
+
+			// Loop through all files.
+			int NumDone = 0;
+			for (const FString& FileName : FoundFiles)
+			{
+				// Update notification with current status.
+				int PercentageDone = (100 * NumDone) / FoundFiles.Num();
+				Async(EAsyncExecution::TaskGraphMainThread, [ConfirmNotification, PercentageDone]()
+				{
+					if (ConfirmNotification.IsValid())
+					{
+						ConfirmNotification->SetText(
+							FText::Format(LOCTEXT("ImgMediaCompleted", "ImgMedia Completed {0}%"),
+								FText::AsNumber(PercentageDone)));
+					}
+				});
+				NumDone++;
+
+				FString FullFileName = FPaths::Combine(SequencePath, FileName);
+
+				// Load image into buffer.
+				TArray64<uint8> InputBuffer;
+				if (!FFileHelper::LoadFileToArray(InputBuffer, *FullFileName))
+				{
+					UE_LOG(LogImgMediaEditor, Error, TEXT("Failed to load %s"), *FullFileName);
+					break;
+				}
+				if (!ImageWrapper.IsValid() || !ImageWrapper->SetCompressed(InputBuffer.GetData(), InputBuffer.Num()))
+				{
+					UE_LOG(LogImgMediaEditor, Error, TEXT("Failed to create image wrapper for %s"), *FullFileName);
+					break;
+				}
+
+				// Get image data.
+				ERGBFormat Format = ImageWrapper->GetFormat();
+				int32 Width = ImageWrapper->GetWidth();
+				int32 Height = ImageWrapper->GetHeight();
+				int32 BitDepth = ImageWrapper->GetBitDepth();
+				TArray64<uint8> RawData;
+				ImageWrapper->GetRaw(Format, BitDepth, RawData);
+
+				// Save image.
+				FString Name = FPaths::Combine(OutPath, FileName);
+				ImageWrapper->SetRaw(RawData.GetData(), RawData.Num(), Width, Height, Format, BitDepth);
+				const TArray64<uint8> CompressedData = ImageWrapper->GetCompressed((int32)EImageCompressionQuality::Uncompressed);
+				FFileHelper::SaveArrayToFile(CompressedData, *Name);
+			}
+		}
+	}
+
+	// Close notification. Must be run on the main thread.
+	Async(EAsyncExecution::TaskGraphMainThread, [ConfirmNotification]()
+	{
+		if (ConfirmNotification.IsValid())
+		{
+			ConfirmNotification->SetEnabled(false);
+			ConfirmNotification->SetCompletionState(SNotificationItem::CS_Success);
+			ConfirmNotification->ExpireAndFadeout();
+		}
+	});
+}
+
+
+#undef LOCTEXT_NAMESPACE
diff --git a/Engine/Plugins/Media/ImgMedia/Source/ImgMediaEditor/Private/Customizations/ImgMediaSourceCustomizationImportInfo.h b/Engine/Plugins/Media/ImgMedia/Source/ImgMediaEditor/Private/Customizations/ImgMediaSourceCustomizationImportInfo.h
new file mode 100644
--- /dev/null
+++ b/Engine/Plugins/Media/ImgMedia/Source/ImgMediaEditor/Private/Customizations/ImgMediaSourceCustomizationImportInfo.h
@@ -0,0 +1,70 @@
+// Copyright Epic Games, Inc. All Rights Reserved.
+
+#pragma once
+
+#include "CoreTypes.h"
+#include "IPropertyTypeCustomization.h"
+#include "Misc/Optional.h"
+#include "Templates/SharedPointer.h"
+
+class IPropertyHandle;
+class SNotificationItem;
+
+/**
+ * Implements a customization for the FImgMediaSourceImportData class.
+ */
+class FImgMediaSourceCustomizationImportInfo
+	: public IPropertyTypeCustomization
+{
+public:
+
+	/** IPropertyTypeCustomization interface */
+	virtual void CustomizeHeader(TSharedRef<class IPropertyHandle> StructPropertyHandle, class FDetailWidgetRow& HeaderRow, IPropertyTypeCustomizationUtils& StructCustomizationUtils) override;
+	virtual void CustomizeChildren(TSharedRef<class IPropertyHandle> StructPropertyHandle, class IDetailChildrenBuilder& StructBuilder, IPropertyTypeCustomizationUtils& StructCustomizationUtils) override;
+
+	/**
+	 * Creates an instance of this class.
+	 *
+	 * @return The new instance.
+	 */
+	static TSharedRef<IPropertyTypeCustomization> MakeInstance()
+	{
+		return MakeShareable(new FImgMediaSourceCustomizationImportInfo());
+	}
+
+private:
+	/**
+	 * Gets the current tile width of each image in pixels.
+	 *
+	 * @return Width.
+	 */
+	TOptional<int32> GetTileWidth() const;
+
+	/**
+	 * Sets the tile width of each image in pixels.
+	 * 
+	 * @param InWidth Width to set.
+	 */
+	void SetTileWidth(int32 InWidth);
+
+	/**
+	 * Called when the import sequence button is clicked.
+	 * 
+	 * @return Whether this handled the event or not.
+	 */
+	FReply OnImportClicked();
+
+	/**
+	 * Async function to import all files in the sequence and generate tiles/mips.
+	 *
+	 * @param SequencePath				Path to look for files.
+	 * @param ConfirmNotification		Notification that will be updated with progress and closed when done.
+	 */
+	static void ImportFiles(const FString& SequencePath, TSharedPtr<SNotificationItem> ConfirmNotification);
+
+	/** Stores our property. */
+	TSharedPtr<IPropertyHandle> PropertyHandle;
+
+	/** Tile width for each image in pixels. */
+	int32 TileWidth = 0;
+};
diff --git a/Engine/Plugins/Media/ImgMedia/Source/ImgMediaEditor/Private/ImgMediaEditorModule.cpp b/Engine/Plugins/Media/ImgMedia/Source/ImgMediaEditor/Private/ImgMediaEditorModule.cpp
--- a/Engine/Plugins/Media/ImgMedia/Source/ImgMediaEditor/Private/ImgMediaEditorModule.cpp
+++ b/Engine/Plugins/Media/ImgMedia/Source/ImgMediaEditor/Private/ImgMediaEditorModule.cpp
@@ -7,6 +7,7 @@
 #include "AssetToolsModule.h"
 #include "AssetTools/ImgMediaSourceActions.h"
 #include "Customizations/ImgMediaSourceCustomization.h"
+#include "Customizations/ImgMediaSourceCustomizationImportInfo.h"
 #include "IAssetTools.h"
 #include "ImgMediaSource.h"
 #include "PropertyEditorModule.h"
@@ -42,11 +43,13 @@ protected:
 	void RegisterCustomizations()
 	{
 		CustomizedStructName = FImgMediaSourceCustomizationSequenceProxy::StaticStruct()->GetFName();
+		ImportInfoStructName = FImgMediaSourceImportInfo::StaticStruct()->GetFName();
 
 		FPropertyEditorModule& PropertyModule = FModuleManager::LoadModuleChecked<FPropertyEditorModule>("PropertyEditor");
 		{
 #if WITH_EDITORONLY_DATA
 			PropertyModule.RegisterCustomPropertyTypeLayout(CustomizedStructName, FOnGetPropertyTypeCustomizationInstance::CreateStatic(&FImgMediaSourceCustomization::MakeInstance));
+			PropertyModule.RegisterCustomPropertyTypeLayout(ImportInfoStructName, FOnGetPropertyTypeCustomizationInstance::CreateStatic(&FImgMediaSourceCustomizationImportInfo::MakeInstance));
 #endif // WITH_EDITORONLY_DATA
 		}
 	}
@@ -57,6 +60,7 @@ protected:
 		FPropertyEditorModule& PropertyModule = FModuleManager::LoadModuleChecked<FPropertyEditorModule>("PropertyEditor");
 		{
 #if WITH_EDITORONLY_DATA
+			PropertyModule.UnregisterCustomPropertyTypeLayout(ImportInfoStructName);
 			PropertyModule.UnregisterCustomPropertyTypeLayout(CustomizedStructName);
 #endif // WITH_EDITORONLY_DATA
 		}
@@ -90,6 +94,7 @@ private:
 
 	/** Customization name to avoid reusing staticstruct during shutdown. */
 	FName CustomizedStructName;
+	FName ImportInfoStructName;
 
 	/** The collection of registered asset type actions. */
 	TArray<TSharedRef<IAssetTypeActions>> RegisteredAssetTypeActions;
