ImgMedia: Tiles can now be created when importing images.
#rnx

diff --git a/Engine/Plugins/Media/ImgMedia/Source/ImgMediaEditor/Private/Customizations/ImgMediaSourceCustomizationImportInfo.cpp b/Engine/Plugins/Media/ImgMedia/Source/ImgMediaEditor/Private/Customizations/ImgMediaSourceCustomizationImportInfo.cpp
--- a/Engine/Plugins/Media/ImgMedia/Source/ImgMediaEditor/Private/Customizations/ImgMediaSourceCustomizationImportInfo.cpp
+++ b/Engine/Plugins/Media/ImgMedia/Source/ImgMediaEditor/Private/Customizations/ImgMediaSourceCustomizationImportInfo.cpp
@@ -70,7 +70,10 @@ void FImgMediaSourceCustomizationImportInfo::CustomizeHeader(TSharedRef<IPropert
 
 void FImgMediaSourceCustomizationImportInfo::CustomizeChildren(TSharedRef<IPropertyHandle> InStructPropertyHandle, IDetailChildrenBuilder& StructBuilder, IPropertyTypeCustomizationUtils& StructCustomizationUtils)
 {
+	// Add tiles.
 	IDetailGroup& TileGroup = StructBuilder.AddGroup(FName(TEXT("TilesGroup")), LOCTEXT("TileGroup_DisplayName", "Tiles"));
+	
+	// Add tile width.
 	TileGroup.AddWidgetRow()
 		.NameContent()
 		[
@@ -84,6 +87,21 @@ void FImgMediaSourceCustomizationImportInfo::CustomizeChildren(TSharedRef<IPrope
 				.Value(this, &FImgMediaSourceCustomizationImportInfo::GetTileWidth)
 				.OnValueChanged(this, &FImgMediaSourceCustomizationImportInfo::SetTileWidth)
 		];
+
+	// Add tile height.
+	TileGroup.AddWidgetRow()
+		.NameContent()
+		[
+			SNew(STextBlock)
+			.Text(LOCTEXT("ImportTileHeight", "Height"))
+			.Font(IDetailLayoutBuilder::GetDetailFont())
+		]
+		.ValueContent()
+		[
+			SNew(SNumericEntryBox<int32>)
+				.Value(this, &FImgMediaSourceCustomizationImportInfo::GetTileHeight)
+				.OnValueChanged(this, &FImgMediaSourceCustomizationImportInfo::SetTileHeight)
+		];
 }
 
 /* FImgMediaSourceCustomizationImport implementation
@@ -99,6 +117,16 @@ void FImgMediaSourceCustomizationImportInfo::SetTileWidth(int32 InWidth)
 	TileWidth = InWidth;
 }
 
+TOptional<int32> FImgMediaSourceCustomizationImportInfo::GetTileHeight() const
+{
+	return TileHeight;
+}
+
+void FImgMediaSourceCustomizationImportInfo::SetTileHeight(int32 InHeight)
+{
+	TileHeight = InHeight;
+}
+
 FReply FImgMediaSourceCustomizationImportInfo::OnImportClicked()
 {
 	// Create notification.
@@ -110,13 +138,15 @@ FReply FImgMediaSourceCustomizationImportInfo::OnImportClicked()
 	Async(EAsyncExecution::Thread, [this, ConfirmNotification]()
 	{
 		FString SequencePath = FImgMediaSourceCustomization::GetSequencePathFromChildProperty(PropertyHandle);
-		ImportFiles(SequencePath, ConfirmNotification);
+		ImportFiles(SequencePath, ConfirmNotification, TileWidth, TileHeight);
 	});
 
 	return FReply::Handled();
 }
 
-void FImgMediaSourceCustomizationImportInfo::ImportFiles(const FString& SequencePath, TSharedPtr<SNotificationItem> ConfirmNotification)
+void FImgMediaSourceCustomizationImportInfo::ImportFiles(const FString& SequencePath,
+	TSharedPtr<SNotificationItem> ConfirmNotification,
+	int32 InTileWidth, int32 InTileHeight)
 {
 	// Create output directory.
 	FString OutPath = FPaths::Combine(SequencePath, TEXT("Imported"));
@@ -179,19 +209,9 @@ void FImgMediaSourceCustomizationImportInfo::ImportFiles(const FString& Sequence
 					break;
 				}
 
-				// Get image data.
-				ERGBFormat Format = ImageWrapper->GetFormat();
-				int32 Width = ImageWrapper->GetWidth();
-				int32 Height = ImageWrapper->GetHeight();
-				int32 BitDepth = ImageWrapper->GetBitDepth();
-				TArray64<uint8> RawData;
-				ImageWrapper->GetRaw(Format, BitDepth, RawData);
-
-				// Save image.
-				FString Name = FPaths::Combine(OutPath, FileName);
-				ImageWrapper->SetRaw(RawData.GetData(), RawData.Num(), Width, Height, Format, BitDepth);
-				const TArray64<uint8> CompressedData = ImageWrapper->GetCompressed((int32)EImageCompressionQuality::Uncompressed);
-				FFileHelper::SaveArrayToFile(CompressedData, *Name);
+				// Import this image.
+				FString Name = FPaths::ChangeExtension(FPaths::Combine(OutPath, FileName), TEXT(""));
+				ImportImage(ImageWrapper, InTileWidth, InTileHeight, Name, Ext);
 			}
 		}
 	}
@@ -208,5 +228,69 @@ void FImgMediaSourceCustomizationImportInfo::ImportFiles(const FString& Sequence
 	});
 }
 
+void FImgMediaSourceCustomizationImportInfo::ImportImage(
+	TSharedPtr<IImageWrapper>& InImageWrapper, 
+	int32 InTileWidth, int32 InTileHeight, const FString& InName, const FString& FileExtension)
+{
+	// Get image data.
+	ERGBFormat Format = InImageWrapper->GetFormat();
+	int32 Width = InImageWrapper->GetWidth();
+	int32 Height = InImageWrapper->GetHeight();
+	int32 BitDepth = InImageWrapper->GetBitDepth();
+	TArray64<uint8> RawData;
+	InImageWrapper->GetRaw(Format, BitDepth, RawData);
+
+	int32 NumTilesX = InTileWidth > 0 ? Width / InTileWidth : 1;
+	int32 NumTilesY = InTileHeight > 0 ? Height / InTileHeight : 1;
+	int32 TileWidth = Width / NumTilesX;
+	int32 TileHeight = Height / NumTilesY;
+	int32 BytesPerPixel = RawData.Num() / (Width * Height);
+	TArray64<uint8> TileRawData;
+	TileRawData.AddZeroed(TileWidth * TileHeight * BytesPerPixel);
+	bool bIsTiled = (NumTilesX > 1) || (NumTilesY > 1);
+
+	// Create a directory if we have tiles.
+	FString FileName;
+	if (bIsTiled)
+	{
+		IPlatformFile& PlatformFile = FPlatformFileManager::Get().GetPlatformFile();
+		PlatformFile.CreateDirectoryTree(*InName);
+		FileName = FPaths::Combine(InName, FPaths::GetCleanFilename(InName));;
+	}
+	else
+	{
+		FileName = InName;
+	}
+	
+	// Loop over y tiles.
+	for (int TileY = 0; TileY < NumTilesY; ++TileY)
+	{
+		// Loop over x tiles.
+		for (int TileX = 0; TileX < NumTilesX; ++TileX)
+		{
+			// Copy tile line by line.
+			uint8* DestPtr = TileRawData.GetData();
+			uint8* SrcPtr = RawData.GetData() + TileX * TileWidth * BytesPerPixel +
+				TileY * TileHeight * Width * BytesPerPixel;
+			for (int LineY = 0; LineY < TileHeight; ++LineY)
+			{
+				FMemory::Memcpy(DestPtr, SrcPtr,
+					TileWidth * BytesPerPixel);
+				DestPtr += TileWidth * BytesPerPixel;
+				SrcPtr += Width * BytesPerPixel;
+			}
+			
+			// Compress data.
+			InImageWrapper->SetRaw(TileRawData.GetData(), TileRawData.Num(),
+				TileWidth, TileHeight, Format, BitDepth);
+			const TArray64<uint8> CompressedData = InImageWrapper->GetCompressed((int32)EImageCompressionQuality::Uncompressed);
+			
+			// Write out tile.
+			FString Name = FString::Format(TEXT("{0}_x{1}_y{2}.{3}"),
+				{*FileName, TileX, TileY, *FileExtension});
+			FFileHelper::SaveArrayToFile(CompressedData, *Name);
+		}
+	}
+}
 
 #undef LOCTEXT_NAMESPACE
diff --git a/Engine/Plugins/Media/ImgMedia/Source/ImgMediaEditor/Private/Customizations/ImgMediaSourceCustomizationImportInfo.h b/Engine/Plugins/Media/ImgMedia/Source/ImgMediaEditor/Private/Customizations/ImgMediaSourceCustomizationImportInfo.h
--- a/Engine/Plugins/Media/ImgMedia/Source/ImgMediaEditor/Private/Customizations/ImgMediaSourceCustomizationImportInfo.h
+++ b/Engine/Plugins/Media/ImgMedia/Source/ImgMediaEditor/Private/Customizations/ImgMediaSourceCustomizationImportInfo.h
@@ -7,6 +7,7 @@
 #include "Misc/Optional.h"
 #include "Templates/SharedPointer.h"
 
+class IImageWrapper;
 class IPropertyHandle;
 class SNotificationItem;
 
@@ -47,6 +48,20 @@ private:
 	 */
 	void SetTileWidth(int32 InWidth);
 
+	/**
+	 * Gets the current tile height of each image in pixels.
+	 *
+	 * @return Height.
+	 */
+	TOptional<int32> GetTileHeight() const;
+
+	/**
+	 * Sets the tile height of each image in pixels.
+	 *
+	 * @param InHeight Height to set.
+	 */
+	void SetTileHeight(int32 InHeight);
+
 	/**
 	 * Called when the import sequence button is clicked.
 	 * 
@@ -57,14 +72,34 @@ private:
 	/**
 	 * Async function to import all files in the sequence and generate tiles/mips.
 	 *
-	 * @param SequencePath				Path to look for files.
-	 * @param ConfirmNotification		Notification that will be updated with progress and closed when done.
+	 * @param SequencePath			Path to look for files.
+	 * @param ConfirmNotification	Notification that will be updated with progress and closed when done.
+	 * @param InTileWidth			Desired width of tiles. 
+	 * @param InTileHeight			Desired height of tiles.
+	 */
+	static void ImportFiles(const FString& SequencePath,
+		TSharedPtr<SNotificationItem> ConfirmNotification,
+		int32 InTileWidth, int32 InTileHeight);
+
+	/**
+	 * Imports a single image and writes out 1 or more files.
+	 * Tiles and mips may be generated.
+	 * This does NOT run on the game thread.
+	 * 
+	 * @param InImageWrapper	ImageWrapper to read/write the image.
+	 * @param InTileWidth		Desired width of tiles. 
+	 * @param InTileHeight		Desired height of tiles.
+	 * @param InName			Full path and name of file to write to WITHOUT the extension (e.g. no .exr)
+	 * @param FileExtension		Extension to append to the file name.
 	 */
-	static void ImportFiles(const FString& SequencePath, TSharedPtr<SNotificationItem> ConfirmNotification);
+	static void ImportImage(TSharedPtr<IImageWrapper>& InImageWrapper,
+		int32 InTileWidth, int32 InTileHeight, const FString& InName, const FString& FileExtension);
 
 	/** Stores our property. */
 	TSharedPtr<IPropertyHandle> PropertyHandle;
 
 	/** Tile width for each image in pixels. */
 	int32 TileWidth = 0;
+	/** Tile height for each image in pixels. */
+	int32 TileHeight = 0;
 };
