ImgMedia: Fixed exr gpu reader buffers trying to reference the reader after the reader has been destroyed.
#rnx

diff --git a/Engine/Plugins/Media/ImgMedia/Source/ImgMedia/Private/Readers/ExrImgMediaReader.cpp b/Engine/Plugins/Media/ImgMedia/Source/ImgMedia/Private/Readers/ExrImgMediaReader.cpp
--- a/Engine/Plugins/Media/ImgMedia/Source/ImgMedia/Private/Readers/ExrImgMediaReader.cpp
+++ b/Engine/Plugins/Media/ImgMedia/Source/ImgMedia/Private/Readers/ExrImgMediaReader.cpp
@@ -179,7 +179,7 @@ TSharedPtr<IImgMediaReader, ESPMode::ThreadSafe> FExrImgMediaReader::GetReader(c
 		&& CVarEnableUncompressedExrGpuReader.GetValueOnAnyThread()
 		)
 	{
-		return MakeShareable(new FExrImgMediaReaderGpu(InLoader));
+		return MakeShared<FExrImgMediaReaderGpu, ESPMode::ThreadSafe>(InLoader);
 	}
 #endif
 	return MakeShareable(new FExrImgMediaReader(InLoader));
diff --git a/Engine/Plugins/Media/ImgMedia/Source/ImgMedia/Private/Readers/ExrImgMediaReaderGpu.cpp b/Engine/Plugins/Media/ImgMedia/Source/ImgMedia/Private/Readers/ExrImgMediaReaderGpu.cpp
--- a/Engine/Plugins/Media/ImgMedia/Source/ImgMedia/Private/Readers/ExrImgMediaReaderGpu.cpp
+++ b/Engine/Plugins/Media/ImgMedia/Source/ImgMedia/Private/Readers/ExrImgMediaReaderGpu.cpp
@@ -374,7 +374,7 @@ SIZE_T FExrImgMediaReaderGpu::GetBufferSize(const FIntPoint& Dim, int32 NumChann
 FStructuredBufferPoolItemSharedPtr FExrImgMediaReaderGpu::AllocateGpuBufferFromPool(uint32 AllocSize, bool bWait)
 {
 	// This function is attached to the shared pointer and is used to return any allocated memory to staging pool.
-	auto BufferDeleter = [AllocSize, this](FStructuredBufferPoolItem* ObjectToDelete) {
+	auto BufferDeleter = [AllocSize](FStructuredBufferPoolItem* ObjectToDelete) {
 		ReturnGpuBufferToStagingPool(AllocSize, ObjectToDelete);
 	};
 
@@ -397,6 +397,7 @@ FStructuredBufferPoolItemSharedPtr FExrImgMediaReaderGpu::AllocateGpuBufferFromP
 		volatile bool bInitDone = false;
 		{
 			AllocatedBuffer = MakeShareable(new FStructuredBufferPoolItem(), MoveTemp(BufferDeleter));
+			AllocatedBuffer->Reader = AsShared();
 
 			// Allocate and unlock the structured buffer on render thread.
 			ENQUEUE_RENDER_COMMAND(CreatePooledBuffer)([AllocatedBuffer, AllocSize, &bInitDone, this, bWait](FRHICommandListImmediate& RHICmdList)
@@ -428,10 +429,12 @@ FStructuredBufferPoolItemSharedPtr FExrImgMediaReaderGpu::AllocateGpuBufferFromP
 
 void FExrImgMediaReaderGpu::ReturnGpuBufferToStagingPool(uint32 AllocSize, FStructuredBufferPoolItem* Buffer)
 {
+	TSharedPtr< FExrImgMediaReaderGpu, ESPMode::ThreadSafe> Reader = Buffer->Reader.Pin();
+
 	// If reader is being deleted, we don't need to return the memory into staging buffer and instead should delete it.
-	if (bIsShuttingDown)
+	if ((Reader.IsValid() == false) || (Reader->bIsShuttingDown))
 	{
-		ENQUEUE_RENDER_COMMAND(DeletePooledBuffers)([this, Buffer](FRHICommandListImmediate& RHICmdList)
+		ENQUEUE_RENDER_COMMAND(DeletePooledBuffers)([Buffer](FRHICommandListImmediate& RHICmdList)
 		{
 			SCOPED_DRAW_EVENT(RHICmdList, FExrImgMediaReaderGpu_ReleaseBuffer);
 
@@ -443,10 +446,10 @@ void FExrImgMediaReaderGpu::ReturnGpuBufferToStagingPool(uint32 AllocSize, FStru
 	}
 	else
 	{
-		FScopeLock ScopeLock(&AllocatorCriticalSecion);
+		FScopeLock ScopeLock(&Reader->AllocatorCriticalSecion);
 
 		// We don't need to process this pooled buffer if the Reader is being destroyed.
-		StagingMemoryPool.Add(AllocSize, Buffer);
+		Reader->StagingMemoryPool.Add(AllocSize, Buffer);
 	}
 
 }
diff --git a/Engine/Plugins/Media/ImgMedia/Source/ImgMedia/Private/Readers/ExrImgMediaReaderGpu.h b/Engine/Plugins/Media/ImgMedia/Source/ImgMedia/Private/Readers/ExrImgMediaReaderGpu.h
--- a/Engine/Plugins/Media/ImgMedia/Source/ImgMedia/Private/Readers/ExrImgMediaReaderGpu.h
+++ b/Engine/Plugins/Media/ImgMedia/Source/ImgMedia/Private/Readers/ExrImgMediaReaderGpu.h
@@ -9,6 +9,7 @@
 #include "IImgMediaReader.h"
 #include "IMediaTextureSampleConverter.h"
 
+class FExrImgMediaReaderGpu;
 struct FImgMediaTileSelection;
 
 struct FStructuredBufferPoolItem
@@ -33,6 +34,11 @@ struct FStructuredBufferPoolItem
 	* is currently using it.
 	*/
 	bool bWillBeSignaled = false;
+
+	/**
+	* Keep track of our reader in case it gets destroyed.
+	*/
+	TWeakPtr<FExrImgMediaReaderGpu, ESPMode::ThreadSafe> Reader;
 };
 
 /**
@@ -43,7 +49,9 @@ typedef TSharedPtr<FStructuredBufferPoolItem, ESPMode::ThreadSafe> FStructuredBu
 /**
  * Implements a reader for EXR image sequences.
  */
-class FExrImgMediaReaderGpu : public FExrImgMediaReader
+class FExrImgMediaReaderGpu
+	: public FExrImgMediaReader
+	, public TSharedFromThis<FExrImgMediaReaderGpu, ESPMode::ThreadSafe>
 {
 public:
 
@@ -85,7 +93,7 @@ public:
 	FStructuredBufferPoolItemSharedPtr AllocateGpuBufferFromPool(uint32 AllocSize, bool bWait = true);
 
 	/** Either return or Add new chunk of memory to the pool based on its size. */
-	void ReturnGpuBufferToStagingPool(uint32 AllocSize, FStructuredBufferPoolItem* Buffer);
+	static void ReturnGpuBufferToStagingPool(uint32 AllocSize, FStructuredBufferPoolItem* Buffer);
 
 	/** Transfer from Staging buffer to Memory pool. */
 	void TransferFromStagingBuffer();
