ImgMedia: User can set the import destination path for images.
#rnx

diff --git a/Engine/Plugins/Media/ImgMedia/Source/ImgMedia/Public/ImgMediaSource.h b/Engine/Plugins/Media/ImgMedia/Source/ImgMedia/Public/ImgMediaSource.h
--- a/Engine/Plugins/Media/ImgMedia/Source/ImgMedia/Public/ImgMediaSource.h
+++ b/Engine/Plugins/Media/ImgMedia/Source/ImgMedia/Public/ImgMediaSource.h
@@ -27,6 +27,14 @@ USTRUCT(BlueprintType)
 struct IMGMEDIA_API FImgMediaSourceImportInfo
 {
 	GENERATED_BODY()
+
+	/** Where to place the imported images. */
+	UPROPERTY(BlueprintReadOnly, EditAnywhere, Category = Sequence)
+	FDirectoryPath DestinationPath;
+
+	/** True if we are overriding destination path and not using the default. */
+	UPROPERTY(BlueprintReadOnly, EditAnywhere, Category = Sequence)
+	bool bIsDestinationPathOverriden = false;
 };
 
 /**
diff --git a/Engine/Plugins/Media/ImgMedia/Source/ImgMediaEditor/Private/Customizations/ImgMediaSourceCustomizationImportInfo.cpp b/Engine/Plugins/Media/ImgMedia/Source/ImgMediaEditor/Private/Customizations/ImgMediaSourceCustomizationImportInfo.cpp
--- a/Engine/Plugins/Media/ImgMedia/Source/ImgMediaEditor/Private/Customizations/ImgMediaSourceCustomizationImportInfo.cpp
+++ b/Engine/Plugins/Media/ImgMedia/Source/ImgMediaEditor/Private/Customizations/ImgMediaSourceCustomizationImportInfo.cpp
@@ -23,6 +23,7 @@
 #include "Widgets/DeclarativeSyntaxSupport.h"
 #include "Widgets/Images/SImage.h"
 #include "Widgets/Input/SButton.h"
+#include "Widgets/Input/SDirectoryPicker.h"
 #include "Widgets/Input/SEditableTextBox.h"
 #include "Widgets/Input/SFilePathPicker.h"
 #include "Widgets/Input/SNumericEntryBox.h"
@@ -66,10 +67,42 @@ void FImgMediaSourceCustomizationImportInfo::CustomizeHeader(TSharedRef<IPropert
 						]
 				]
 		];
+
+	// Get destination path property.
+	TSharedPtr<IPropertyHandle> DestinationProperty = PropertyHandle->GetChildHandle("DestinationPath");
+	if (DestinationProperty.IsValid())
+	{
+		DestinationPathPropertyHandle = DestinationProperty->GetChildHandle("Path");
+	}
+
+	// Get destination path property.
+	IsDestinationPathOverridenPropertyHandle = PropertyHandle->GetChildHandle("bIsDestinationPathOverriden");
 }
 
 void FImgMediaSourceCustomizationImportInfo::CustomizeChildren(TSharedRef<IPropertyHandle> InStructPropertyHandle, IDetailChildrenBuilder& StructBuilder, IPropertyTypeCustomizationUtils& StructCustomizationUtils)
 {
+	// Add destination path.
+	// Get the destination path if we are overriding.
+	FString DestinationPath;
+	if (IsDestinationPathOverriden())
+	{
+		DestinationPath = GetDestinationPath();
+	}
+
+	StructBuilder.AddCustomRow(LOCTEXT("DestinationPath", "Destination Path"))
+		.NameContent()
+		[
+			SNew(STextBlock)
+				.Text(LOCTEXT("DestinationPath", "Destination Path"))
+				.Font(IDetailLayoutBuilder::GetDetailFont())
+		]
+		.ValueContent()
+		[
+			SNew(SDirectoryPicker)
+				.Directory(DestinationPath)
+				.OnDirectoryChanged(this, &FImgMediaSourceCustomizationImportInfo::OnDestinationPathChanged)
+		];
+
 	// Add tiles.
 	IDetailGroup& TileGroup = StructBuilder.AddGroup(FName(TEXT("TilesGroup")), LOCTEXT("TileGroup_DisplayName", "Tiles"));
 	
@@ -127,29 +160,131 @@ void FImgMediaSourceCustomizationImportInfo::SetTileHeight(int32 InHeight)
 	TileHeight = InHeight;
 }
 
+FString FImgMediaSourceCustomizationImportInfo::GetDestinationPath()
+{
+	FString Path;
+
+	if (DestinationPathPropertyHandle.IsValid())
+	{
+		if (DestinationPathPropertyHandle->GetValue(Path) != FPropertyAccess::Success)
+		{
+			UE_LOG(LogImgMediaEditor, Error, TEXT("Could not get value for destination path."));
+		}
+	}
+	else
+	{
+		UE_LOG(LogImgMediaEditor, Error, TEXT("Could not get property for destination path."));
+	}
+
+	return Path;
+}
+
+void FImgMediaSourceCustomizationImportInfo::SetDestinationPath(const FString& InPath)
+{
+	if (DestinationPathPropertyHandle.IsValid())
+	{
+		if (DestinationPathPropertyHandle->SetValue(InPath) != FPropertyAccess::Success)
+		{
+			UE_LOG(LogImgMediaEditor, Error, TEXT("Could not set value for %s"),
+				*DestinationPathPropertyHandle->GetPropertyDisplayName().ToString());
+		}
+	}
+	else
+	{
+		UE_LOG(LogImgMediaEditor, Error, TEXT("Could not get property for %s."),
+			*DestinationPathPropertyHandle->GetPropertyDisplayName().ToString());
+	}
+}
+
+bool FImgMediaSourceCustomizationImportInfo::IsDestinationPathOverriden()
+{
+	bool bIsDestinationPathOverriden = false;
+	if (IsDestinationPathOverridenPropertyHandle.IsValid())
+	{
+		if (IsDestinationPathOverridenPropertyHandle->GetValue(bIsDestinationPathOverriden) != FPropertyAccess::Success)
+		{
+			UE_LOG(LogImgMediaEditor, Error, TEXT("Could not get value for %s"),
+				*IsDestinationPathOverridenPropertyHandle->GetPropertyDisplayName().ToString());
+		}
+	}
+	else
+	{
+		UE_LOG(LogImgMediaEditor, Error, TEXT("Could not get property for %s"),
+			*IsDestinationPathOverridenPropertyHandle->GetPropertyDisplayName().ToString());
+	}
+
+	return bIsDestinationPathOverriden;
+}
+
+void FImgMediaSourceCustomizationImportInfo::SetIsDestinationPathOverriden(bool bInIsOverriden)
+{
+	if (IsDestinationPathOverridenPropertyHandle.IsValid())
+	{
+		if (IsDestinationPathOverridenPropertyHandle->SetValue(bInIsOverriden) != FPropertyAccess::Success)
+		{
+			UE_LOG(LogImgMediaEditor, Error, TEXT("Could not set value for %s"),
+				*IsDestinationPathOverridenPropertyHandle->GetPropertyDisplayName().ToString());
+		}
+	}
+	else
+	{
+		UE_LOG(LogImgMediaEditor, Error, TEXT("Could not get property for %s"),
+			*IsDestinationPathOverridenPropertyHandle->GetPropertyDisplayName().ToString());
+	}
+}
+
+void FImgMediaSourceCustomizationImportInfo::OnDestinationPathChanged(const FString& Directory)
+{
+	// We are overriding the path if the directory is not empty.
+	bool bIsPathOverriden = (Directory.IsEmpty() == false);
+
+	// If we are not overriding the path now, and we were not overriding before, then don't update
+	// as we would just lose any default path that has been set.
+	if ((bIsPathOverriden) || (IsDestinationPathOverriden()))
+	{
+		SetDestinationPath(Directory);
+	}
+
+	SetIsDestinationPathOverriden(bIsPathOverriden);
+}
+
 FReply FImgMediaSourceCustomizationImportInfo::OnImportClicked()
 {
+	// Get destination path.
+	FString DestinationPath = GetDestinationPath();
+	
+	// Do we want to override it?
+	bool bOverrideDestinationPath = IsDestinationPathOverriden();
+	if (bOverrideDestinationPath == false)
+	{
+		// Nope. Derive it from the Sequence path.
+		FString SequencePath = FImgMediaSourceCustomization::GetSequencePathFromChildProperty(PropertyHandle);
+		DestinationPath = FPaths::Combine(SequencePath, TEXT("Imported"));
+		SetDestinationPath(DestinationPath);
+	}
+
 	// Create notification.
 	FNotificationInfo Info(FText::GetEmpty());
 	Info.bFireAndForget = false;
 	TSharedPtr<SNotificationItem> ConfirmNotification = FSlateNotificationManager::Get().AddNotification(Info);
 
 	// Start async task to import files.
-	Async(EAsyncExecution::Thread, [this, ConfirmNotification]()
+	Async(EAsyncExecution::Thread, [this, DestinationPath, ConfirmNotification]()
 	{
 		FString SequencePath = FImgMediaSourceCustomization::GetSequencePathFromChildProperty(PropertyHandle);
-		ImportFiles(SequencePath, ConfirmNotification, TileWidth, TileHeight);
+		ImportFiles(SequencePath, DestinationPath, ConfirmNotification, TileWidth, TileHeight);
 	});
 
 	return FReply::Handled();
 }
 
 void FImgMediaSourceCustomizationImportInfo::ImportFiles(const FString& SequencePath,
+	const FString& InDestinationPath,
 	TSharedPtr<SNotificationItem> ConfirmNotification,
 	int32 InTileWidth, int32 InTileHeight)
 {
 	// Create output directory.
-	FString OutPath = FPaths::Combine(SequencePath, TEXT("Imported"));
+	FString OutPath = InDestinationPath;
 	IPlatformFile& PlatformFile = FPlatformFileManager::Get().GetPlatformFile();
 	PlatformFile.CreateDirectoryTree(*OutPath);
 
diff --git a/Engine/Plugins/Media/ImgMedia/Source/ImgMediaEditor/Private/Customizations/ImgMediaSourceCustomizationImportInfo.h b/Engine/Plugins/Media/ImgMedia/Source/ImgMediaEditor/Private/Customizations/ImgMediaSourceCustomizationImportInfo.h
--- a/Engine/Plugins/Media/ImgMedia/Source/ImgMediaEditor/Private/Customizations/ImgMediaSourceCustomizationImportInfo.h
+++ b/Engine/Plugins/Media/ImgMedia/Source/ImgMediaEditor/Private/Customizations/ImgMediaSourceCustomizationImportInfo.h
@@ -62,6 +62,31 @@ private:
 	 */
 	void SetTileHeight(int32 InHeight);
 
+	/**
+	 * Gets where the imported files should go.
+	 */
+	FString GetDestinationPath();
+
+	/**
+	 * Sets where the imported files should go.
+	 */
+	void SetDestinationPath(const FString& InPath);
+
+	/**
+	 * Query if the destination path has been overriden.
+	 */
+	bool IsDestinationPathOverriden();
+
+	/**
+	 * Set if the destination path has been overriden.
+	 */
+	void SetIsDestinationPathOverriden(bool bInIsOverriden);
+
+	/**
+	 *Called when DestinationPath changes.
+	 */
+	void OnDestinationPathChanged(const FString& Directory);
+
 	/**
 	 * Called when the import sequence button is clicked.
 	 * 
@@ -78,6 +103,7 @@ private:
 	 * @param InTileHeight			Desired height of tiles.
 	 */
 	static void ImportFiles(const FString& SequencePath,
+		const FString& DestinationPath,
 		TSharedPtr<SNotificationItem> ConfirmNotification,
 		int32 InTileWidth, int32 InTileHeight);
 
@@ -97,6 +123,10 @@ private:
 
 	/** Stores our property. */
 	TSharedPtr<IPropertyHandle> PropertyHandle;
+	/** Storess the property to the destination path.*/
+	TSharedPtr<IPropertyHandle> DestinationPathPropertyHandle;
+	/** Stores the property that says if the destination path is overriden. */
+	TSharedPtr<IPropertyHandle> IsDestinationPathOverridenPropertyHandle;
 
 	/** Tile width for each image in pixels. */
 	int32 TileWidth = 0;
