ImgMedia: Added button to allow usage of imported data.
#rnx

diff --git a/Engine/Plugins/Media/ImgMedia/Source/ImgMedia/Public/ImgMediaSource.h b/Engine/Plugins/Media/ImgMedia/Source/ImgMedia/Public/ImgMediaSource.h
--- a/Engine/Plugins/Media/ImgMedia/Source/ImgMedia/Public/ImgMediaSource.h
+++ b/Engine/Plugins/Media/ImgMedia/Source/ImgMedia/Public/ImgMediaSource.h
@@ -35,6 +35,10 @@ struct IMGMEDIA_API FImgMediaSourceImportInfo
 	/** True if we are overriding destination path and not using the default. */
 	UPROPERTY(BlueprintReadOnly, EditAnywhere, Category = Sequence)
 	bool bIsDestinationPathOverriden = false;
+
+	/** True if the imported imagess are to be used. */
+	UPROPERTY(BlueprintReadOnly, EditAnywhere, Category = Sequence)
+	bool bIsUsable = false;
 };
 
 /**
diff --git a/Engine/Plugins/Media/ImgMedia/Source/ImgMediaEditor/Private/Customizations/ImgMediaSourceCustomizationImportInfo.cpp b/Engine/Plugins/Media/ImgMedia/Source/ImgMediaEditor/Private/Customizations/ImgMediaSourceCustomizationImportInfo.cpp
--- a/Engine/Plugins/Media/ImgMedia/Source/ImgMediaEditor/Private/Customizations/ImgMediaSourceCustomizationImportInfo.cpp
+++ b/Engine/Plugins/Media/ImgMedia/Source/ImgMediaEditor/Private/Customizations/ImgMediaSourceCustomizationImportInfo.cpp
@@ -77,10 +77,18 @@ void FImgMediaSourceCustomizationImportInfo::CustomizeHeader(TSharedRef<IPropert
 
 	// Get destination path property.
 	IsDestinationPathOverridenPropertyHandle = PropertyHandle->GetChildHandle("bIsDestinationPathOverriden");
+	// Get the usable property.
+	IsUsablePropertyHandle = PropertyHandle->GetChildHandle("bIsUsable");
 }
 
 void FImgMediaSourceCustomizationImportInfo::CustomizeChildren(TSharedRef<IPropertyHandle> InStructPropertyHandle, IDetailChildrenBuilder& StructBuilder, IPropertyTypeCustomizationUtils& StructCustomizationUtils)
 {
+	// Add IsUsable button.
+	if (IsUsablePropertyHandle.IsValid())
+	{
+		StructBuilder.AddProperty(IsUsablePropertyHandle.ToSharedRef());
+	}
+
 	// Add destination path.
 	// Get the destination path if we are overriding.
 	FString DestinationPath;
@@ -198,38 +206,58 @@ void FImgMediaSourceCustomizationImportInfo::SetDestinationPath(const FString& I
 
 bool FImgMediaSourceCustomizationImportInfo::IsDestinationPathOverriden()
 {
-	bool bIsDestinationPathOverriden = false;
-	if (IsDestinationPathOverridenPropertyHandle.IsValid())
+	return GetGenericBoolProperty(IsDestinationPathOverridenPropertyHandle);
+}
+
+void FImgMediaSourceCustomizationImportInfo::SetIsDestinationPathOverriden(bool bInIsOverriden)
+{
+	SetGenericBoolProperty(IsDestinationPathOverridenPropertyHandle, bInIsOverriden);
+}
+
+bool FImgMediaSourceCustomizationImportInfo::IsUsableProperty()
+{
+	return GetGenericBoolProperty(IsUsablePropertyHandle);
+}
+
+void FImgMediaSourceCustomizationImportInfo::SetIsUsableProperty(bool bInIsUsable)
+{
+	SetGenericBoolProperty(IsUsablePropertyHandle, bInIsUsable);
+}
+
+bool FImgMediaSourceCustomizationImportInfo::GetGenericBoolProperty(const TSharedPtr<IPropertyHandle>& InProperty)
+{
+	bool bIsTrue = false;
+	if (InProperty.IsValid())
 	{
-		if (IsDestinationPathOverridenPropertyHandle->GetValue(bIsDestinationPathOverriden) != FPropertyAccess::Success)
+		if (InProperty->GetValue(bIsTrue) != FPropertyAccess::Success)
 		{
 			UE_LOG(LogImgMediaEditor, Error, TEXT("Could not get value for %s"),
-				*IsDestinationPathOverridenPropertyHandle->GetPropertyDisplayName().ToString());
+				*InProperty->GetPropertyDisplayName().ToString());
 		}
 	}
 	else
 	{
 		UE_LOG(LogImgMediaEditor, Error, TEXT("Could not get property for %s"),
-			*IsDestinationPathOverridenPropertyHandle->GetPropertyDisplayName().ToString());
+			*InProperty->GetPropertyDisplayName().ToString());
 	}
 
-	return bIsDestinationPathOverriden;
+	return bIsTrue;
 }
 
-void FImgMediaSourceCustomizationImportInfo::SetIsDestinationPathOverriden(bool bInIsOverriden)
+void FImgMediaSourceCustomizationImportInfo::SetGenericBoolProperty(const TSharedPtr<IPropertyHandle>& InProperty, bool bInIsTrue)
 {
-	if (IsDestinationPathOverridenPropertyHandle.IsValid())
+	if (InProperty.IsValid())
 	{
-		if (IsDestinationPathOverridenPropertyHandle->SetValue(bInIsOverriden) != FPropertyAccess::Success)
+		if (InProperty->SetValue(bInIsTrue) != FPropertyAccess::Success)
 		{
 			UE_LOG(LogImgMediaEditor, Error, TEXT("Could not set value for %s"),
-				*IsDestinationPathOverridenPropertyHandle->GetPropertyDisplayName().ToString());
+				*InProperty->GetPropertyDisplayName().ToString());
 		}
 	}
 	else
 	{
 		UE_LOG(LogImgMediaEditor, Error, TEXT("Could not get property for %s"),
-			*IsDestinationPathOverridenPropertyHandle->GetPropertyDisplayName().ToString());
+			*InProperty->GetPropertyDisplayName().ToString());
 	}
 }
 
@@ -242,7 +270,14 @@ void FImgMediaSourceCustomizationImportInfo::OnDestinationPathChanged(const FStr
 	// as we would just lose any default path that has been set.
 	if ((bIsPathOverriden) || (IsDestinationPathOverriden()))
 	{
-		SetDestinationPath(Directory);
+		// Did the path change?
+		FString CurrentDestinationNormalized = GetDestinationPath() / TEXT("");
+		FString DirectoryNormalized = Directory / TEXT("");
+		if (FPaths::IsSamePath(CurrentDestinationNormalized, DirectoryNormalized) == false)
+		{
+			SetDestinationPath(Directory);
+			SetIsUsableProperty(false);
+		}
 	}
 
 	SetIsDestinationPathOverriden(bIsPathOverriden);
@@ -263,6 +298,9 @@ FReply FImgMediaSourceCustomizationImportInfo::OnImportClicked()
 		SetDestinationPath(DestinationPath);
 	}
 
+	// Mark that we can use this now.
+	SetIsUsableProperty(true);
+
 	// Create notification.
 	FNotificationInfo Info(FText::GetEmpty());
 	Info.bFireAndForget = false;
diff --git a/Engine/Plugins/Media/ImgMedia/Source/ImgMediaEditor/Private/Customizations/ImgMediaSourceCustomizationImportInfo.h b/Engine/Plugins/Media/ImgMedia/Source/ImgMediaEditor/Private/Customizations/ImgMediaSourceCustomizationImportInfo.h
--- a/Engine/Plugins/Media/ImgMedia/Source/ImgMediaEditor/Private/Customizations/ImgMediaSourceCustomizationImportInfo.h
+++ b/Engine/Plugins/Media/ImgMedia/Source/ImgMediaEditor/Private/Customizations/ImgMediaSourceCustomizationImportInfo.h
@@ -83,7 +83,27 @@ private:
 	void SetIsDestinationPathOverriden(bool bInIsOverriden);
 
 	/**
-	 *Called when DestinationPath changes.
+	 * Query if we should use the imported images.
+	 */
+	bool IsUsableProperty();
+
+	/**
+	 * Set if we should use the imported images.
+	 */
+	void SetIsUsableProperty(bool bInIsUsable);
+
+	/**
+	 * Generic query function for boolean properties.
+	 */
+	bool GetGenericBoolProperty(const TSharedPtr<IPropertyHandle>& InProperty);
+
+	/**
+	 * Generic setter for boolean properties.
+	 */
+	void SetGenericBoolProperty(const TSharedPtr<IPropertyHandle>& InProperty, bool bInIsTrue);
+
+	/**
+	 * Called when DestinationPath changes.
 	 */
 	void OnDestinationPathChanged(const FString& Directory);
 
@@ -127,6 +147,8 @@ private:
 	TSharedPtr<IPropertyHandle> DestinationPathPropertyHandle;
 	/** Stores the property that says if the destination path is overriden. */
 	TSharedPtr<IPropertyHandle> IsDestinationPathOverridenPropertyHandle;
+	/** Stores the property that says if the imported files are usable. */
+	TSharedPtr<IPropertyHandle> IsUsablePropertyHandle;
 
 	/** Tile width for each image in pixels. */
 	int32 TileWidth = 0;
