ExrImgMediaReaderGpu (ImageMedia): Fixed an issue with Reader falling back to default reader when frame read is simply cancelled, but still successful.

diff --git a/Engine/Plugins/Media/ImgMedia/Source/ImgMedia/Private/Readers/ExrImgMediaReaderGpu.cpp b/Engine/Plugins/Media/ImgMedia/Source/ImgMedia/Private/Readers/ExrImgMediaReaderGpu.cpp
--- a/Engine/Plugins/Media/ImgMedia/Source/ImgMedia/Private/Readers/ExrImgMediaReaderGpu.cpp
+++ b/Engine/Plugins/Media/ImgMedia/Source/ImgMedia/Private/Readers/ExrImgMediaReaderGpu.cpp
@@ -180,7 +180,7 @@ bool FExrImgMediaReaderGpu::ReadFrame(int32 FrameId, int32 MipLevel, const FImgM
 			{
 				// Get for our frame/mip level.
 				const FString& ImagePath = Loader->GetImagePath(FrameId, CurrentMipLevel);
-				bool bResult = false;
+				EReadResult ReadResult = Fail;
 				
 				const SIZE_T BufferSize = GetBufferSize(Dim, NumChannels);
 				FStructuredBufferPoolItemSharedPtr& BufferData = BufferDataArray[CurrentMipLevel];
@@ -188,12 +188,14 @@ bool FExrImgMediaReaderGpu::ReadFrame(int32 FrameId, int32 MipLevel, const FImgM
 				uint16* MipDataPtr = static_cast<uint16*>(BufferData->MappedBuffer);
 
 #if READ_IN_CHUNKS
-				bResult = ReadInChunks(MipDataPtr, ImagePath, FrameId, Dim, BufferSize, PixelSize, NumChannels);
+				ReadResult = ReadInChunks(MipDataPtr, ImagePath, FrameId, Dim, BufferSize, PixelSize, NumChannels);
 #else
-				bResult = FExrReader::GenerateTextureData(MipDataPtr, ImagePath, Dim.X, Dim.Y, PixelSize, NumChannels);
+
+				bool bResult = FExrReader::GenerateTextureData(MipDataPtr, ImagePath, Dim.X, Dim.Y, PixelSize, NumChannels);
+				ReadResult = bResult ? Success : Fail;
 #endif
 
-				if (!bResult)
+				if (ReadResult == Fail)
 				{
 					// Check if we have a compressed file.
 					FRgbaInputFile InputFileMip(ImagePath);
@@ -308,9 +310,9 @@ void FExrImgMediaReaderGpu::OnTick()
 /* FExrImgMediaReaderGpu implementation
  *****************************************************************************/
 
-bool FExrImgMediaReaderGpu::ReadInChunks(uint16* Buffer, const FString& ImagePath, int32 FrameId, const FIntPoint& Dim, int32 BufferSize, int32 PixelSize, int32 NumChannels)
+FExrImgMediaReaderGpu::EReadResult FExrImgMediaReaderGpu::ReadInChunks(uint16* Buffer, const FString& ImagePath, int32 FrameId, const FIntPoint& Dim, int32 BufferSize, int32 PixelSize, int32 NumChannels)
 {
-	bool bResult = true;
+	EReadResult bResult = Success;
 
 	// Chunks are of 16 MB
 	const int32 ChunkSize = 0xF42400;
@@ -323,7 +325,7 @@ bool FExrImgMediaReaderGpu::ReadInChunks(uint16* Buffer, const FString& ImagePat
 
 	if (!ChunkReader.OpenExrAndPrepareForPixelReading(ImagePath, Dim.X, Dim.Y, PixelSize, NumChannels))
 	{
-		return false;
+		return Fail;
 	}
 
 	for (int32 Row = 0; Row <= NumChunks; Row++)
@@ -340,14 +342,14 @@ bool FExrImgMediaReaderGpu::ReadInChunks(uint16* Buffer, const FString& ImagePat
 			if (CanceledFrames.Remove(FrameId) > 0)
 			{
 				UE_LOG(LogImgMedia, Warning, TEXT("Reader %p: Canceling Frame %i At chunk # %i"), this, FrameId, Row);
-				bResult = false;
+				bResult = Cancelled;
 				break;
 			}
 		}
 
 		if (!ChunkReader.ReadExrImageChunk(reinterpret_cast<char*>(Buffer) + CurrentBufferPos, Step))
 		{
-			bResult = false;
+			bResult = Fail;
 			break;
 		}
 		CurrentBufferPos += Step;
@@ -355,7 +357,7 @@ bool FExrImgMediaReaderGpu::ReadInChunks(uint16* Buffer, const FString& ImagePat
 
 	if (!ChunkReader.CloseExrFile())
 	{
-		return false;
+		return Fail;
 	}
 
 	return bResult;
diff --git a/Engine/Plugins/Media/ImgMedia/Source/ImgMedia/Private/Readers/ExrImgMediaReaderGpu.h b/Engine/Plugins/Media/ImgMedia/Source/ImgMedia/Private/Readers/ExrImgMediaReaderGpu.h
--- a/Engine/Plugins/Media/ImgMedia/Source/ImgMedia/Private/Readers/ExrImgMediaReaderGpu.h
+++ b/Engine/Plugins/Media/ImgMedia/Source/ImgMedia/Private/Readers/ExrImgMediaReaderGpu.h
@@ -73,11 +73,18 @@ public:
 
 protected:
 
+	enum EReadResult
+	{
+		Fail,
+		Success,
+		Cancelled
+	};
+
 	/** 
 	 * This function reads file in 16 MB chunks and if it detects that
 	 * Frame is pending for cancellation stops reading the file and returns false.
 	*/
-	bool ReadInChunks(uint16* Buffer, const FString& ImagePath, int32 FrameId, const FIntPoint& Dim, int32 BufferSize, int32 PixelSize, int32 NumChannels);
+	EReadResult ReadInChunks(uint16* Buffer, const FString& ImagePath, int32 FrameId, const FIntPoint& Dim, int32 BufferSize, int32 PixelSize, int32 NumChannels);
 
 	/**
 	 * Get the size of the buffer needed to load in an image.
