ImgMedia: Loader can use tiles.
#rnx

diff --git a/Engine/Plugins/Media/ImgMedia/Source/ImgMedia/Private/Assets/ImgMediaSource.cpp b/Engine/Plugins/Media/ImgMedia/Source/ImgMedia/Private/Assets/ImgMediaSource.cpp
--- a/Engine/Plugins/Media/ImgMedia/Source/ImgMedia/Private/Assets/ImgMediaSource.cpp
+++ b/Engine/Plugins/Media/ImgMedia/Source/ImgMedia/Private/Assets/ImgMediaSource.cpp
@@ -189,6 +189,12 @@ void UImgMediaSource::PostEditChangeProperty(FPropertyChangedEvent& PropertyChan
 
 FString UImgMediaSource::GetFullPath() const
 {
+	// Use the import path?
+	if ((ImportInfo.bIsUsable) && (ImportInfo.DestinationPath.Path.IsEmpty() == false))
+	{
+		return ImportInfo.DestinationPath.Path;
+	}
+
 	if (!FPaths::IsRelative(SequencePath.Path))
 	{
 		return SequencePath.Path;
diff --git a/Engine/Plugins/Media/ImgMedia/Source/ImgMedia/Private/Loader/ImgMediaLoader.cpp b/Engine/Plugins/Media/ImgMedia/Source/ImgMedia/Private/Loader/ImgMediaLoader.cpp
--- a/Engine/Plugins/Media/ImgMedia/Source/ImgMedia/Private/Loader/ImgMediaLoader.cpp
+++ b/Engine/Plugins/Media/ImgMedia/Source/ImgMedia/Private/Loader/ImgMediaLoader.cpp
@@ -838,6 +838,13 @@ void FImgMediaLoader::FindFiles(const FString& SequencePath, TArray<FString>& Ou
 	TArray<FString> FoundFiles;
 	IFileManager::Get().FindFiles(FoundFiles, *SequencePath, TEXT("*"));
 
+	// If we did not find any files, then maybe we have tile directories.
+	bool bIsTiled = FoundFiles.Num() == 0;
+	if (bIsTiled)
+	{
+		IFileManager::Get().FindFiles(FoundFiles, *(SequencePath / TEXT("*")), false, true);
+	}
+
 	UE_LOG(LogImgMedia, Verbose, TEXT("Loader %p: Found %i image files in %s"), this, FoundFiles.Num(), *SequencePath);
 
 	FoundFiles.Sort();
@@ -865,7 +872,15 @@ void FImgMediaLoader::FindFiles(const FString& SequencePath, TArray<FString>& Ou
 
 			LastIndex = ThisIndex;
 		}
-		OutputPaths.Add(FPaths::Combine(SequencePath, File));
+
+		FString FullPath = FPaths::Combine(SequencePath, File);
+		if (bIsTiled)
+		{
+			// If we have tiles, then use the first tile as our reference.
+			FString TiledFileName = File + TEXT("_x0_y0.exr");
+			FullPath = FPaths::Combine(FullPath, TiledFileName);
+		}
+		OutputPaths.Add(FullPath);
 	}
 
 }
