ImgMedia: Added number of tiles to ImgMediaSource, and this gets set when importing images.
#rnx

diff --git a/Engine/Plugins/Media/ImgMedia/Source/ImgMedia/Public/ImgMediaSource.h b/Engine/Plugins/Media/ImgMedia/Source/ImgMedia/Public/ImgMediaSource.h
--- a/Engine/Plugins/Media/ImgMedia/Source/ImgMedia/Public/ImgMediaSource.h
+++ b/Engine/Plugins/Media/ImgMedia/Source/ImgMedia/Public/ImgMediaSource.h
@@ -39,6 +39,14 @@ struct IMGMEDIA_API FImgMediaSourceImportInfo
 	/** True if the imported imagess are to be used. */
 	UPROPERTY(BlueprintReadOnly, EditAnywhere, Category = Sequence)
 	bool bIsUsable = false;
+
+	/** Number of tiles in the X direction. If 0, then there are no tiles. */
+	UPROPERTY(BlueprintReadOnly, EditAnywhere, Category = Sequence)
+	int32 NumTilesX = 0;
+
+	/** Number if tiles in the Y direction. If 0, then there are no tiles. */
+	UPROPERTY(BlueprintReadOnly, EditAnywhere, Category = Sequence)
+	int32 NumTilesY = 0;
 };
 
 /**
diff --git a/Engine/Plugins/Media/ImgMedia/Source/ImgMediaEditor/Private/Customizations/ImgMediaSourceCustomizationImportInfo.cpp b/Engine/Plugins/Media/ImgMedia/Source/ImgMediaEditor/Private/Customizations/ImgMediaSourceCustomizationImportInfo.cpp
--- a/Engine/Plugins/Media/ImgMedia/Source/ImgMediaEditor/Private/Customizations/ImgMediaSourceCustomizationImportInfo.cpp
+++ b/Engine/Plugins/Media/ImgMedia/Source/ImgMediaEditor/Private/Customizations/ImgMediaSourceCustomizationImportInfo.cpp
@@ -79,6 +79,19 @@ void FImgMediaSourceCustomizationImportInfo::CustomizeHeader(TSharedRef<IPropert
 	IsDestinationPathOverridenPropertyHandle = PropertyHandle->GetChildHandle("bIsDestinationPathOverriden");
 	// Get the usable property.
 	IsUsablePropertyHandle = PropertyHandle->GetChildHandle("bIsUsable");
+
+	// Get the ImgMediaSource we are editing.
+	ImgMediaSource = nullptr;
+	TArray<UObject*> OuterObjects;
+	PropertyHandle->GetOuterObjects(OuterObjects);
+	if (OuterObjects.Num() > 0)
+	{
+		UObject* Obj = OuterObjects[0];
+		if (Obj->IsA<UImgMediaSource>())
+		{
+			ImgMediaSource = CastChecked<UImgMediaSource>(Obj);
+		}
+	}
 }
 
 void FImgMediaSourceCustomizationImportInfo::CustomizeChildren(TSharedRef<IPropertyHandle> InStructPropertyHandle, IDetailChildrenBuilder& StructBuilder, IPropertyTypeCustomizationUtils& StructCustomizationUtils)
@@ -143,6 +156,20 @@ void FImgMediaSourceCustomizationImportInfo::CustomizeChildren(TSharedRef<IPrope
 				.Value(this, &FImgMediaSourceCustomizationImportInfo::GetTileHeight)
 				.OnValueChanged(this, &FImgMediaSourceCustomizationImportInfo::SetTileHeight)
 		];
+
+	// Add num tiles x.
+	TSharedPtr<IPropertyHandle> NumTilesX = PropertyHandle->GetChildHandle("NumTilesX");
+	if (NumTilesX.IsValid())
+	{
+		TileGroup.AddPropertyRow(NumTilesX.ToSharedRef());
+	}
+
+	// Add num tiles y.
+	TSharedPtr<IPropertyHandle> NumTilesY = PropertyHandle->GetChildHandle("NumTilesY");
+	if (NumTilesY.IsValid())
+	{
+		TileGroup.AddPropertyRow(NumTilesY.ToSharedRef());
+	}
 }
 
 /* FImgMediaSourceCustomizationImport implementation
@@ -352,6 +379,7 @@ void FImgMediaSourceCustomizationImportInfo::ImportFiles(const FString& Sequence
 
 			// Loop through all files.
 			int NumDone = 0;
+			bool bNeedSetUp = true;
 			for (const FString& FileName : FoundFiles)
 			{
 				// Update notification with current status.
@@ -382,6 +410,28 @@ void FImgMediaSourceCustomizationImportInfo::ImportFiles(const FString& Sequence
 					break;
 				}
 
+				// Get info for the sequence.
+				if (bNeedSetUp)
+				{
+					bNeedSetUp = false;
+
+					// Get number of tiles.
+					int32 Width = ImageWrapper->GetWidth();
+					int32 Height = ImageWrapper->GetHeight();
+					int32 NumTilesX = InTileWidth > 0 ? Width / InTileWidth : 1;
+					int32 NumTilesY = InTileHeight > 0 ? Height / InTileHeight : 1;
+					if (ImgMediaSource.IsValid())
+					{
+						if ((ImgMediaSource->ImportInfo.NumTilesX != NumTilesX) ||
+							(ImgMediaSource->ImportInfo.NumTilesY != NumTilesY))
+						{
+							ImgMediaSource->ImportInfo.NumTilesX = NumTilesX;
+							ImgMediaSource->ImportInfo.NumTilesY = NumTilesY;
+							ImgMediaSource->MarkPackageDirty();
+						}
+					}
+				}
+
 				// Import this image.
 				FString Name = FPaths::ChangeExtension(FPaths::Combine(OutPath, FileName), TEXT(""));
 				ImportImage(ImageWrapper, InTileWidth, InTileHeight, Name, Ext);
diff --git a/Engine/Plugins/Media/ImgMedia/Source/ImgMediaEditor/Private/Customizations/ImgMediaSourceCustomizationImportInfo.h b/Engine/Plugins/Media/ImgMedia/Source/ImgMediaEditor/Private/Customizations/ImgMediaSourceCustomizationImportInfo.h
--- a/Engine/Plugins/Media/ImgMedia/Source/ImgMediaEditor/Private/Customizations/ImgMediaSourceCustomizationImportInfo.h
+++ b/Engine/Plugins/Media/ImgMedia/Source/ImgMediaEditor/Private/Customizations/ImgMediaSourceCustomizationImportInfo.h
@@ -10,6 +10,7 @@
 class IImageWrapper;
 class IPropertyHandle;
 class SNotificationItem;
+class UImgMediaSource;
 
 /**
  * Implements a customization for the FImgMediaSourceImportData class.
@@ -122,7 +123,7 @@ private:
 	 * @param InTileWidth			Desired width of tiles. 
 	 * @param InTileHeight			Desired height of tiles.
 	 */
-	static void ImportFiles(const FString& SequencePath,
+	void ImportFiles(const FString& SequencePath,
 		const FString& DestinationPath,
 		TSharedPtr<SNotificationItem> ConfirmNotification,
 		int32 InTileWidth, int32 InTileHeight);
@@ -149,6 +150,8 @@ private:
 	TSharedPtr<IPropertyHandle> IsDestinationPathOverridenPropertyHandle;
 	/** Stores the property that says if the imported files are usable. */
 	TSharedPtr<IPropertyHandle> IsUsablePropertyHandle;
+	/** Stores the ImgMediaSource that we are editing. */
+	TWeakObjectPtr<UImgMediaSource> ImgMediaSource;
 
 	/** Tile width for each image in pixels. */
 	int32 TileWidth = 0;
