ImgMedia: The loader gets tile info from the player via the media options.

#rnx

diff --git a/Engine/Plugins/Media/ImgMedia/Source/ImgMedia/Private/Assets/ImgMediaSource.cpp b/Engine/Plugins/Media/ImgMedia/Source/ImgMedia/Private/Assets/ImgMediaSource.cpp
--- a/Engine/Plugins/Media/ImgMedia/Source/ImgMedia/Private/Assets/ImgMediaSource.cpp
+++ b/Engine/Plugins/Media/ImgMedia/Source/ImgMedia/Private/Assets/ImgMediaSource.cpp
@@ -113,6 +113,16 @@ int64 UImgMediaSource::GetMediaOption(const FName& Key, int64 DefaultValue) cons
 		return FrameRateOverride.Numerator;
 	}
 
+	if (Key == ImgMedia::NumTilesXOption)
+	{
+		return ImportInfo.NumTilesX;
+	}
+
+	if (Key == ImgMedia::NumTilesYOption)
+	{
+		return ImportInfo.NumTilesY;
+	}
+
 	return Super::GetMediaOption(Key, DefaultValue);
 }
 
@@ -143,7 +153,9 @@ bool UImgMediaSource::HasMediaOption(const FName& Key) const
 		(Key == ImgMedia::FrameRateOverrideDenonimatorOption) ||
 		(Key == ImgMedia::FrameRateOverrideNumeratorOption) ||
 		(Key == ImgMedia::ProxyOverrideOption) ||
-		(Key == ImgMedia::MipMapInfoOption))
+		(Key == ImgMedia::MipMapInfoOption) ||
+		(Key == ImgMedia::NumTilesXOption) ||
+		(Key == ImgMedia::NumTilesYOption))
 	{
 		return true;
 	}
diff --git a/Engine/Plugins/Media/ImgMedia/Source/ImgMedia/Private/ImgMediaPrivate.h b/Engine/Plugins/Media/ImgMedia/Source/ImgMedia/Private/ImgMediaPrivate.h
--- a/Engine/Plugins/Media/ImgMedia/Source/ImgMedia/Private/ImgMediaPrivate.h
+++ b/Engine/Plugins/Media/ImgMedia/Source/ImgMedia/Private/ImgMediaPrivate.h
@@ -48,4 +48,10 @@ namespace ImgMedia
 
 	/** Name of the MipMapInfo media option. */
 	static const FName MipMapInfoOption("ImgMipMapInfo");
+
+	/** Name of the NumTilesX media option. */
+	static const FName NumTilesXOption("ImgNumTilesX");
+
+	/** Name of the NumTilesY media option. */
+	static const FName NumTilesYOption("ImgNumTilesY");
 }
diff --git a/Engine/Plugins/Media/ImgMedia/Source/ImgMedia/Private/Loader/ImgMediaLoader.cpp b/Engine/Plugins/Media/ImgMedia/Source/ImgMedia/Private/Loader/ImgMediaLoader.cpp
--- a/Engine/Plugins/Media/ImgMedia/Source/ImgMedia/Private/Loader/ImgMediaLoader.cpp
+++ b/Engine/Plugins/Media/ImgMedia/Source/ImgMedia/Private/Loader/ImgMediaLoader.cpp
@@ -63,6 +63,8 @@ FImgMediaLoader::FImgMediaLoader(const TSharedRef<FImgMediaScheduler, ESPMode::T
 	, ImageWrapperModule(FModuleManager::LoadModuleChecked<IImageWrapperModule>("ImageWrapper"))
 	, Initialized(false)
 	, bFillGapsInSequence(bInFillGapsInSequence)
+	, NumTilesX(0)
+	, NumTilesY(0)
 	, NumLoadAhead(0)
 	, NumLoadBehind(0)
 	, Scheduler(InScheduler)
@@ -620,7 +622,8 @@ IQueuedWork* FImgMediaLoader::GetWork()
 }
 
 
-void FImgMediaLoader::Initialize(const FString& SequencePath, const FFrameRate& FrameRateOverride, bool Loop)
+void FImgMediaLoader::Initialize(const FString& SequencePath, const FFrameRate& FrameRateOverride,
+	bool Loop, int32 InNumTilesX, int32 InNumTilesY)
 {
 	UE_LOG(LogImgMedia, Verbose, TEXT("Loader %p: Initializing with %s (FrameRateOverride = %s, Loop = %i)"),
 		this,
@@ -631,6 +634,15 @@ void FImgMediaLoader::Initialize(const FString& SequencePath, const FFrameRate&
 
 	check(!Initialized); // reinitialization not allowed for now
 
+	// Set up tile info.
+	NumTilesX = InNumTilesX;
+	NumTilesY = InNumTilesY;
+	if ((NumTilesX < 1) || (NumTilesY < 1))
+	{
+		NumTilesX = 1;
+		NumTilesY = 1;
+	}
+
 	LoadSequence(SequencePath, FrameRateOverride, Loop);
 	FPlatformMisc::MemoryBarrier();
 
diff --git a/Engine/Plugins/Media/ImgMedia/Source/ImgMedia/Private/Loader/ImgMediaLoader.h b/Engine/Plugins/Media/ImgMedia/Source/ImgMedia/Private/Loader/ImgMediaLoader.h
--- a/Engine/Plugins/Media/ImgMedia/Source/ImgMedia/Private/Loader/ImgMediaLoader.h
+++ b/Engine/Plugins/Media/ImgMedia/Source/ImgMedia/Private/Loader/ImgMediaLoader.h
@@ -203,6 +203,22 @@ public:
 		return ImagePaths.Num() > 0 ? ImagePaths[0].Num() : 0;
 	}
 
+	/**
+	 * Get the number of tiles in the X direction.
+	 */
+	int32 GetNumTilesX() const
+	{
+		return NumTilesX;
+	}
+
+	/**
+	 * Get the number of tiles in the Y direction.
+	 */
+	int32 GetNumTilesY() const
+	{
+		return NumTilesY;
+	}
+
 	/**
 	 * Get the next work item.
 	 *
@@ -218,9 +234,12 @@ public:
 	 * @param SequencePath Path to the image sequence.
 	 * @param FrameRateOverride The frame rate to use (0/0 = do not override).
 	 * @param Loop Whether the cache should loop around.
+	 * @param InNumTilesX Number of tiles in the X direction.
+	 * @param InNumTilesY Number of tiles in the Y direction.
 	 * @see IsInitialized
 	 */
-	void Initialize(const FString& SequencePath, const FFrameRate& FrameRateOverride, bool Loop);
+	void Initialize(const FString& SequencePath, const FFrameRate& FrameRateOverride, bool Loop,
+		int32 InNumTilesX, int32 InNumTilesY);
 
 	/**
 	 * Whether this loader has been initialized yet.
@@ -422,6 +441,12 @@ private:
 	/** If true, then any gaps in the sequence will be filled with blank frames. */
 	bool bFillGapsInSequence;
 
+	/** Number of tiles in the X direction. */
+	int32 NumTilesX;
+
+	/** Number of tiles in the Y direction. */
+	int32 NumTilesY;
+
 	/** The number of frames to load ahead of the play head. */
 	int32 NumLoadAhead;
 
diff --git a/Engine/Plugins/Media/ImgMedia/Source/ImgMedia/Private/Player/ImgMediaPlayer.cpp b/Engine/Plugins/Media/ImgMedia/Source/ImgMedia/Private/Player/ImgMediaPlayer.cpp
--- a/Engine/Plugins/Media/ImgMedia/Source/ImgMedia/Private/Player/ImgMediaPlayer.cpp
+++ b/Engine/Plugins/Media/ImgMedia/Source/ImgMedia/Private/Player/ImgMediaPlayer.cpp
@@ -214,7 +214,12 @@ bool FImgMediaPlayer::Open(const FString& Url, const IMediaOptions* Options)
 
 	const FString SequencePath = Url.RightChop(6);
 
-	Async(EAsyncExecution::ThreadPool, [FrameRateOverride, LoaderPtr = TWeakPtr<FImgMediaLoader, ESPMode::ThreadSafe>(Loader), Proxy, SequencePath, Loop = ShouldLoop]()
+	// Get tile info.
+	int32 NumTilesX = Options->GetMediaOption(ImgMedia::NumTilesXOption, (int64)0);
+	int32 NumTilesY = Options->GetMediaOption(ImgMedia::NumTilesYOption, (int64)0);
+
+	Async(EAsyncExecution::ThreadPool, [FrameRateOverride, LoaderPtr = TWeakPtr<FImgMediaLoader,ESPMode::ThreadSafe>(Loader),
+		Proxy, SequencePath, Loop = ShouldLoop, NumTilesX, NumTilesY]()
 	{
 		TSharedPtr<FImgMediaLoader, ESPMode::ThreadSafe> PinnedLoader = LoaderPtr.Pin();
 
@@ -227,7 +232,7 @@ bool FImgMediaPlayer::Open(const FString& Url, const IMediaOptions* Options)
 				ProxyPath = SequencePath; // fall back to root folder
 			}
 
-			PinnedLoader->Initialize(ProxyPath, FrameRateOverride, Loop);
+			PinnedLoader->Initialize(ProxyPath, FrameRateOverride, Loop, NumTilesX, NumTilesY);
 		}
 	});
 
