ImgMedia: Fixed static analysis warning.
#rnx

diff --git a/Engine/Plugins/Media/ImgMedia/Source/ImgMedia/Private/Player/ImgMediaPlayer.cpp b/Engine/Plugins/Media/ImgMedia/Source/ImgMedia/Private/Player/ImgMediaPlayer.cpp
--- a/Engine/Plugins/Media/ImgMedia/Source/ImgMedia/Private/Player/ImgMediaPlayer.cpp
+++ b/Engine/Plugins/Media/ImgMedia/Source/ImgMedia/Private/Player/ImgMediaPlayer.cpp
@@ -182,11 +182,15 @@ bool FImgMediaPlayer::Open(const FString& Url, const IMediaOptions* Options)
 	FFrameRate FrameRateOverride(0, 0);
 	TSharedPtr<FImgMediaMipMapInfo, ESPMode::ThreadSafe> MipMapInfo;
 	bool bFillGapsInSequence = true;
+	int32 NumTilesX = 0;
+	int32 NumTilesY = 0;
 	if (Options != nullptr)
 	{
 		FrameRateOverride.Denominator = Options->GetMediaOption(ImgMedia::FrameRateOverrideDenonimatorOption, 0LL);
 		FrameRateOverride.Numerator = Options->GetMediaOption(ImgMedia::FrameRateOverrideNumeratorOption, 0LL);
 		bFillGapsInSequence = Options->GetMediaOption(ImgMedia::FillGapsInSequenceOption, true);
+		NumTilesX = Options->GetMediaOption(ImgMedia::NumTilesXOption, (int64)0);
+		NumTilesY = Options->GetMediaOption(ImgMedia::NumTilesYOption, (int64)0);
 
 		TSharedPtr<IMediaOptions::FDataContainer, ESPMode::ThreadSafe> DefaultValue;
 		TSharedPtr<IMediaOptions::FDataContainer, ESPMode::ThreadSafe> DataContainer = Options->GetMediaOption(ImgMedia::MipMapInfoOption, DefaultValue);
@@ -214,10 +218,6 @@ bool FImgMediaPlayer::Open(const FString& Url, const IMediaOptions* Options)
 
 	const FString SequencePath = Url.RightChop(6);
 
-	// Get tile info.
-	int32 NumTilesX = Options->GetMediaOption(ImgMedia::NumTilesXOption, (int64)0);
-	int32 NumTilesY = Options->GetMediaOption(ImgMedia::NumTilesYOption, (int64)0);
-
 	Async(EAsyncExecution::ThreadPool, [FrameRateOverride, LoaderPtr = TWeakPtr<FImgMediaLoader,ESPMode::ThreadSafe>(Loader),
 		Proxy, SequencePath, Loop = ShouldLoop, NumTilesX, NumTilesY]()
 	{
