ImgMedia: EXR CPU reader reads in tiles.

#rnx

diff --git a/Engine/Plugins/Media/ImgMedia/Source/ImgMedia/Private/Readers/ExrImgMediaReader.cpp b/Engine/Plugins/Media/ImgMedia/Source/ImgMedia/Private/Readers/ExrImgMediaReader.cpp
--- a/Engine/Plugins/Media/ImgMedia/Source/ImgMedia/Private/Readers/ExrImgMediaReader.cpp
+++ b/Engine/Plugins/Media/ImgMedia/Source/ImgMedia/Private/Readers/ExrImgMediaReader.cpp
@@ -62,6 +62,11 @@ bool FExrImgMediaReader::ReadFrame(int32 FrameId, int32 MipLevel, const FImgMedi
 		return false;
 	}
 	
+	// Get tile info.
+	int32 NumTilesX = Loader->GetNumTilesX();
+	int32 NumTilesY = Loader->GetNumTilesY();
+	bool bHasTiles = (NumTilesX * NumTilesY) > 1;
+
 	// Do we already have our buffer?
 	if (OutFrame->Data.IsValid() == false)
 	{
@@ -73,6 +78,10 @@ bool FExrImgMediaReader::ReadFrame(int32 FrameId, int32 MipLevel, const FImgMedi
 			return false;
 		}
 
+		// If we have tiles, then this is the size of just a tile, so multiply to get the full size.
+		OutFrame->Info.Dim.X *= NumTilesX;
+		OutFrame->Info.Dim.Y *= NumTilesY;
+
 		const FIntPoint& Dim = OutFrame->Info.Dim;
 
 		if (Dim.GetMin() <= 0)
@@ -121,18 +130,41 @@ bool FExrImgMediaReader::ReadFrame(int32 FrameId, int32 MipLevel, const FImgMedi
 		bool IsThisLevelPresent = (OutFrame->MipMapsPresent & (1 << CurrentMipLevel)) != 0;
 		bool ReadThisMip = (CurrentMipLevel >= MipLevel) &&
 			(IsThisLevelPresent == false);
+		int32 BytesPerPixel = sizeof(uint16) * 4;
 		if (ReadThisMip)
 		{
-			// Get for our frame/mip level.
-			const FString& Image = Loader->GetImagePath(FrameId, CurrentMipLevel);
-			FRgbaInputFile InputFile(Image, 2);
-
-			// read frame data
-			InputFile.SetFrameBuffer(MipDataPtr, Dim);
-			InputFile.ReadPixels(0, Dim.Y - 1);
-
-			OutFrame->MipMapsPresent |= 1 << CurrentMipLevel;
-			LevelFoundSoFar = true;
+			FString Image = Loader->GetImagePath(FrameId, CurrentMipLevel);
+			FString BaseImage;
+			if (bHasTiles)
+			{
+				// Remove "_x0_y0.exr" so we can add on the correct name for the tile we want.
+				BaseImage = Image.LeftChop(10);
+			}
+			int FrameBufferOffsetY = 0;
+			int32 TileWidth = Dim.X / NumTilesX;
+			int32 TileHeight = Dim.Y / NumTilesY;
+			for (int32 TileY = 0; TileY < NumTilesY; TileY++)
+			{
+				int FrameBufferOffsetX = 0;
+				for (int32 TileX = 0; TileX < NumTilesX; TileX++)
+				{
+					// Get for our frame/mip level.
+					if (bHasTiles)
+					{
+						Image = FString::Printf(TEXT("%s_x%d_y%d.exr"), *BaseImage, TileX, TileY);
+					}
+					FRgbaInputFile InputFile(Image, 2);
+
+					// read frame data
+					InputFile.SetFrameBuffer(MipDataPtr + FrameBufferOffsetX + FrameBufferOffsetY, Dim);
+					InputFile.ReadPixels(0, TileHeight - 1);
+
+					FrameBufferOffsetX += TileWidth * BytesPerPixel;
+					OutFrame->MipMapsPresent |= 1 << CurrentMipLevel;
+					LevelFoundSoFar = true;
+				}
+				FrameBufferOffsetY += Dim.X * TileHeight * BytesPerPixel;
+			}
 		}
 
 		// Next level.
