ImgMedia: Cleaned up bytes per pixel calculations in CPU reader.
#rnx

diff --git a/Engine/Plugins/Media/ImgMedia/Source/ImgMedia/Private/Readers/ExrImgMediaReader.cpp b/Engine/Plugins/Media/ImgMedia/Source/ImgMedia/Private/Readers/ExrImgMediaReader.cpp
--- a/Engine/Plugins/Media/ImgMedia/Source/ImgMedia/Private/Readers/ExrImgMediaReader.cpp
+++ b/Engine/Plugins/Media/ImgMedia/Source/ImgMedia/Private/Readers/ExrImgMediaReader.cpp
@@ -67,6 +67,10 @@ bool FExrImgMediaReader::ReadFrame(int32 FrameId, int32 MipLevel, const FImgMedi
 	int32 NumTilesY = Loader->GetNumTilesY();
 	bool bHasTiles = (NumTilesX * NumTilesY) > 1;
 
+	int32 BytesPerPixelPerChannel = sizeof(uint16);
+	int32 NumChannels = 4;
+	int32 BytesPerPixel = BytesPerPixelPerChannel * NumChannels;
+
 	// Do we already have our buffer?
 	if (OutFrame->Data.IsValid() == false)
 	{
@@ -113,10 +117,12 @@ bool FExrImgMediaReader::ReadFrame(int32 FrameId, int32 MipLevel, const FImgMedi
 			FMemory::Free(ObjectToDelete);
 #endif
 		};
+		
+		// The EXR RGBA interface only outputs RGBA data.
 		OutFrame->Format = EMediaTextureSampleFormat::FloatRGBA;
 		OutFrame->Data = MakeShareable(Buffer, MoveTemp(BufferDeleter));
 		OutFrame->MipMapsPresent = 0;
-		OutFrame->Stride = OutFrame->Info.Dim.X * sizeof(unsigned short) * 4;
+		OutFrame->Stride = OutFrame->Info.Dim.X * BytesPerPixel;
 	}
 
 	// Loop over all mips.
@@ -130,7 +136,7 @@ bool FExrImgMediaReader::ReadFrame(int32 FrameId, int32 MipLevel, const FImgMedi
 		bool IsThisLevelPresent = (OutFrame->MipMapsPresent & (1 << CurrentMipLevel)) != 0;
 		bool ReadThisMip = (CurrentMipLevel >= MipLevel) &&
 			(IsThisLevelPresent == false);
-		int32 BytesPerPixel = sizeof(uint16) * 4;
+		
 		if (ReadThisMip)
 		{
 			FString Image = Loader->GetImagePath(FrameId, CurrentMipLevel);
@@ -168,7 +174,7 @@ bool FExrImgMediaReader::ReadFrame(int32 FrameId, int32 MipLevel, const FImgMedi
 		}
 
 		// Next level.
-		MipDataPtr += Dim.X * Dim.Y * sizeof(uint16) * 4;
+		MipDataPtr += Dim.X * Dim.Y * BytesPerPixel;
 		Dim /= 2;
 		if (IsThisLevelPresent)
 		{
