ImgMedia: By default the tile selection has all tiles visible now.

#rnx

diff --git a/Engine/Plugins/Media/ImgMedia/Source/ImgMedia/Public/ImgMediaMipMapInfo.h b/Engine/Plugins/Media/ImgMedia/Source/ImgMedia/Public/ImgMediaMipMapInfo.h
--- a/Engine/Plugins/Media/ImgMedia/Source/ImgMedia/Public/ImgMediaMipMapInfo.h
+++ b/Engine/Plugins/Media/ImgMedia/Source/ImgMedia/Public/ImgMediaMipMapInfo.h
@@ -67,7 +67,7 @@ struct FImgMediaTileSelection
 
 	FImgMediaTileSelection()
 	{
-		SetAllNotVisible();
+		SetAllVisible();
 	}
 
 	/**
