ImgMedia: EXR CPU reader only reads in the tiles requested.
#rnx

diff --git a/Engine/Plugins/Media/ImgMedia/Source/ImgMedia/Private/Readers/ExrImgMediaReader.cpp b/Engine/Plugins/Media/ImgMedia/Source/ImgMedia/Private/Readers/ExrImgMediaReader.cpp
--- a/Engine/Plugins/Media/ImgMedia/Source/ImgMedia/Private/Readers/ExrImgMediaReader.cpp
+++ b/Engine/Plugins/Media/ImgMedia/Source/ImgMedia/Private/Readers/ExrImgMediaReader.cpp
@@ -12,6 +12,7 @@
 #include "UObject/UObjectGlobals.h"
 #include "HardwareInfo.h"
 #include "ImgMediaLoader.h"
+#include "ImgMediaMipMapInfo.h"
 
 DECLARE_MEMORY_STAT(TEXT("EXR Reader Pool Memory."), STAT_ExrMediaReaderPoolMem, STATGROUP_ImgMediaPlugin);
 
@@ -66,6 +67,10 @@ bool FExrImgMediaReader::ReadFrame(int32 FrameId, int32 MipLevel, const FImgMedi
 	int32 NumTilesX = Loader->GetNumTilesX();
 	int32 NumTilesY = Loader->GetNumTilesY();
 	bool bHasTiles = (NumTilesX * NumTilesY) > 1;
+	int32 StartTileX = InTileSelection.TopLeftX;
+	int32 StartTileY = InTileSelection.TopLeftY;
+	int32 EndTileX = FMath::Min((int32)InTileSelection.BottomRightX, NumTilesX);
+	int32 EndTileY = FMath::Min((int32)InTileSelection.BottomRightY, NumTilesY);
 
 	int32 BytesPerPixelPerChannel = sizeof(uint16);
 	int32 NumChannels = 4;
@@ -146,13 +151,21 @@ bool FExrImgMediaReader::ReadFrame(int32 FrameId, int32 MipLevel, const FImgMedi
 				// Remove "_x0_y0.exr" so we can add on the correct name for the tile we want.
 				BaseImage = Image.LeftChop(10);
 			}
-			int FrameBufferOffsetY = 0;
+
 			int32 TileWidth = Dim.X / NumTilesX;
 			int32 TileHeight = Dim.Y / NumTilesY;
-			for (int32 TileY = 0; TileY < NumTilesY; TileY++)
+			// The offset into the frame buffer for each row of tiles.
+			// Total width * height of a tile * bytes per pixel.
+			int32 FrameBufferOffsetPerTileY = Dim.X * TileHeight * BytesPerPixel;
+			int32 FrameBufferOffsetY = FrameBufferOffsetPerTileY * StartTileY;
+			
+			for (int32 TileY = StartTileY; TileY < EndTileY; TileY++)
 			{
-				int FrameBufferOffsetX = 0;
-				for (int32 TileX = 0; TileX < NumTilesX; TileX++)
+				// The offset into the frame buffer for each column of tiles.
+				// Tile width * bytes per pixel.
+				int32 FrameBufferOffsetPerTileX = TileWidth * BytesPerPixel;
+				int32 FrameBufferOffsetX = FrameBufferOffsetPerTileX * StartTileX;
+				for (int32 TileX = StartTileX; TileX < EndTileX; TileX++)
 				{
 					// Get for our frame/mip level.
 					if (bHasTiles)
@@ -165,11 +178,12 @@ bool FExrImgMediaReader::ReadFrame(int32 FrameId, int32 MipLevel, const FImgMedi
 					InputFile.SetFrameBuffer(MipDataPtr + FrameBufferOffsetX + FrameBufferOffsetY, Dim);
 					InputFile.ReadPixels(0, TileHeight - 1);
 
-					FrameBufferOffsetX += TileWidth * BytesPerPixel;
+					FrameBufferOffsetX += FrameBufferOffsetPerTileX;
 					OutFrame->MipMapsPresent |= 1 << CurrentMipLevel;
 					LevelFoundSoFar = true;
 				}
-				FrameBufferOffsetY += Dim.X * TileHeight * BytesPerPixel;
+				
+				FrameBufferOffsetY += FrameBufferOffsetPerTileY;
 			}
 		}
 
