ImgMedia: Exposed global cache.
#rnx

diff --git a/Engine/Plugins/Media/ImgMedia/Source/ImgMedia/Private/GlobalCache/ImgMediaGlobalCache.h b/Engine/Plugins/Media/ImgMedia/Source/ImgMedia/Private/GlobalCache/ImgMediaGlobalCache.h
deleted file mode 100644
diff --git a/Engine/Plugins/Media/ImgMedia/Source/ImgMedia/Public/IImgMediaModule.h b/Engine/Plugins/Media/ImgMedia/Source/ImgMedia/Public/IImgMediaModule.h
--- a/Engine/Plugins/Media/ImgMedia/Source/ImgMedia/Public/IImgMediaModule.h
+++ b/Engine/Plugins/Media/ImgMedia/Source/ImgMedia/Public/IImgMediaModule.h
@@ -13,7 +13,7 @@ class IMediaPlayer;
 /**
  * Interface for the ImgMedia module.
  */
-class IImgMediaModule
+class IMGMEDIA_API IImgMediaModule
 	: public IModuleInterface
 {
 public:
diff --git a/Engine/Plugins/Media/ImgMedia/Source/ImgMedia/Public/ImgMediaGlobalCache.h b/Engine/Plugins/Media/ImgMedia/Source/ImgMedia/Public/ImgMediaGlobalCache.h
new file mode 100644
--- /dev/null
+++ b/Engine/Plugins/Media/ImgMedia/Source/ImgMedia/Public/ImgMediaGlobalCache.h
@@ -0,0 +1,198 @@
+// Copyright Epic Games, Inc. All Rights Reserved.
+
+#pragma once
+
+#include "CoreTypes.h"
+#include "Containers/Map.h"
+#include "Delegates/IDelegateInstance.h"
+#include "Templates/SharedPointer.h"
+
+struct FFileChangeData; 
+struct FImgMediaFrame;
+class UImgMediaSettings;
+
+/**
+ * A global cache for all ImgMedia players.
+ *
+ * Uses Least Recently Used (LRU).
+ */
+class IMGMEDIA_API FImgMediaGlobalCache : public TSharedFromThis<FImgMediaGlobalCache, ESPMode::ThreadSafe>
+{
+public:
+	/**
+	 * Constructor.
+	 */
+	FImgMediaGlobalCache();
+
+	/**
+	 * Desctructor.
+	  */
+	~FImgMediaGlobalCache();
+
+	/**
+	 * Initialize the cache.
+	 */
+	void Initialize();
+
+	/**
+	 * Shut down the cache.
+	 */
+	void Shutdown();
+
+	/**
+	 * Adds a frame to the cache.
+	 *
+	 * @param	FileName		Filename of frame to add.
+	 * @param	Sequence		Indentifying name of this sequence.
+	 * @param	Index			Index of frame to add.
+	 * @param	Frame			Actual frame to add.
+	 * @param	HasMipMaps		Whether this frame contains mipmaps or not.
+	 */
+	void AddFrame(const FString& FileName, const FName& Sequence, int32 Index, const TSharedPtr<FImgMediaFrame, ESPMode::ThreadSafe>& Frame, bool HasMipMaps);
+
+	/**
+	 * Find the entry with the specified sequence and index and mark it as the most recently used.
+	 *
+	 * @param	Sequence		Indentifying name of this sequence.
+	 * @param	Index			Index of frame to touch.
+	 */
+	TSharedPtr<FImgMediaFrame, ESPMode::ThreadSafe>* FindAndTouch(const FName& Sequence, int32 Index);
+
+	/**
+	 * Find the indices of all cached entries of a particular sequence.
+	 *
+	 * @param	Sequence		Indentifying name of this sequence.
+	 * @param OutIndices Will contain the collection of indices.
+	 */
+	void GetIndices(const FName& Sequence, TArray<int32>& OutIndices) const;
+
+	/**
+	 * Empties the cache.
+	 */
+	void EmptyCache();
+
+private:
+
+	/** An entry in the cache. */
+	struct FImgMediaGlobalCacheEntry
+	{
+		/** File name of this frame. */
+		FString FileName;
+		/** Size of this frame in bytes. */
+		SIZE_T FrameSize;
+		/** Frame index. */
+		int32 Index;
+		/** Actual frame. */
+		TSharedPtr<FImgMediaFrame, ESPMode::ThreadSafe> Frame;
+
+		/** Previous entry in the cache. */
+		FImgMediaGlobalCacheEntry* LessRecent;
+		/** Previous entry in the cache that is part of the same sequence. */
+		FImgMediaGlobalCacheEntry* LessRecentSequence;
+		/** Next entry in the cache. */
+		FImgMediaGlobalCacheEntry* MoreRecent;
+		/** Next entry in the cache that is part of the same sequence. */
+		FImgMediaGlobalCacheEntry* MoreRecentSequence;
+
+		/** Constructor. */
+		FImgMediaGlobalCacheEntry(const FString& InFileName, SIZE_T InFrameSize, int32 InIndex, const TSharedPtr<FImgMediaFrame, ESPMode::ThreadSafe>& InFrame)
+			: FileName(InFileName)
+			, FrameSize(InFrameSize)
+			, Index(InIndex)
+			, Frame(InFrame)
+			, LessRecent(nullptr)
+			, LessRecentSequence(nullptr)
+			, MoreRecent(nullptr)
+			, MoreRecentSequence(nullptr)
+		{
+		}
+	};
+
+	/** Entry that was used first. */
+	FImgMediaGlobalCacheEntry* LeastRecent;
+	/** Entry that was used last. */
+	FImgMediaGlobalCacheEntry* MostRecent;
+
+	/** Maps a sequence name to the most recent cache entry of that sequence. */
+	TMap<FName, FImgMediaGlobalCacheEntry*> MapSequenceToMostRecentEntry;
+	/** Maps a cache entry that is the least recent entry to the name of its sequence. */
+	TMap<FImgMediaGlobalCacheEntry*, FName> MapLeastRecentToSequence;
+	/** Maps a sequence name and frame index to an entry in the cache. */
+	TMap<TPair<FName, int32>, FImgMediaGlobalCacheEntry*> MapFrameToEntry;
+
+	/** Current size of the cache in bytes. */
+	SIZE_T CurrentSize;
+	/** Maximum size of the cache in bytes. */
+	SIZE_T MaxSize;
+
+	/** Critical section for synchronizing access to Frames. */
+	mutable FCriticalSection CriticalSection;
+
+#if WITH_EDITORONLY_DATA
+
+	/** Handle to registered UpdateSettings delegate. */
+	FDelegateHandle UpdateSettingsDelegateHandle;
+
+#endif
+
+#if WITH_EDITOR
+
+	/** Maps a sequence directory to a directory watcher. */
+	TMap<FString, FDelegateHandle> MapSequenceToWatcher;
+	/** Maps a file name to the index of that file in the sequence. */
+	TMap<FString, int32> MapFileToIndex;
+
+#endif // WITH_EDITOR
+	
+	/**
+	 * Empties the cache until the current size + Extra <= the max size of the cache.
+	 */
+	void EnforceMaxSize(SIZE_T Extra);
+
+	/**
+	 * Removes an entry from the cache and deletes the entry.
+	 *
+	 * @param	Sequence		Indentifying name of this sequence.
+	 * @param	Entry			Entry to remove.
+	 */
+	void Remove(const FName& Sequence, FImgMediaGlobalCacheEntry& Entry);
+
+	/**
+	 * Inserts an entry into the cache as the most recent entry.
+	 * Assumes the entry is not already in the cache.
+	 *
+	 * @param	Sequence		Indentifying name of this sequence.
+	 * @param	Entry			Entry to mark as recent.
+	 */
+	void MarkAsRecent(const FName& Sequence, FImgMediaGlobalCacheEntry& Entry);
+
+	/**
+	 * Removes an entry from the cache but does not delete the entry.
+	 *
+	 * @param	Sequence		Indentifying name of this sequence.
+	 * @param	Entry			Entry to remove.
+	 */
+	void Unlink(const FName& Sequence, FImgMediaGlobalCacheEntry& Entry);
+
+	/**
+	 * Empties the cache.
+	 * 
+	 * Assumes CriticalSection is already locked.
+	 */
+	void Empty();
+
+	/** 
+	 * Updates our settings from ImgMediaAsettings. 
+	 */
+	void UpdateSettings(const UImgMediaSettings* Settings);
+
+#if WITH_EDITOR
+
+	/**
+	 * Called when one of our source files changes.
+	 */
+	void OnDirectoryChanged(const TArray<FFileChangeData>& InFileChanges);
+
+#endif // WITH_EDITOR
+};
+
