ImgMedia: MipMapInfo is also used by tiled images.
#rnx

diff --git a/Engine/Plugins/Media/ImgMedia/Source/ImgMedia/Private/Loader/ImgMediaLoader.cpp b/Engine/Plugins/Media/ImgMedia/Source/ImgMedia/Private/Loader/ImgMediaLoader.cpp
--- a/Engine/Plugins/Media/ImgMedia/Source/ImgMedia/Private/Loader/ImgMediaLoader.cpp
+++ b/Engine/Plugins/Media/ImgMedia/Source/ImgMedia/Private/Loader/ImgMediaLoader.cpp
@@ -798,18 +798,21 @@ void FImgMediaLoader::LoadSequence(const FString& SequencePath, const FFrameRate
 	SequenceDuration = FrameNumberToTime(GetNumImages());
 	SIZE_T UncompressedSize = FirstFrameInfo.UncompressedSize;
 
-	// If we have no mips, then get rid of our MipMapInfoObject.
+	// If we have no mips or tiles, then get rid of our MipMapInfoObject.
 	// Otherwise, set it up.
 	if (MipMapInfo.IsValid())
 	{
-		if (GetNumMipLevels() == 1)
+		if ((GetNumMipLevels() == 1) && (IsTiled() == false))
 		{
 			MipMapInfo.Reset();
 		}
 		else
 		{
 			MipMapInfo->SetTextureInfo(SequenceName, GetNumMipLevels(), SequenceDim);
-			UncompressedSize = (UncompressedSize * 4) / 3;
+			if (GetNumMipLevels() > 1)
+			{
+				UncompressedSize = (UncompressedSize * 4) / 3;
+			}
 		}
 	}
 
diff --git a/Engine/Plugins/Media/ImgMedia/Source/ImgMedia/Private/Loader/ImgMediaLoader.h b/Engine/Plugins/Media/ImgMedia/Source/ImgMedia/Private/Loader/ImgMediaLoader.h
--- a/Engine/Plugins/Media/ImgMedia/Source/ImgMedia/Private/Loader/ImgMediaLoader.h
+++ b/Engine/Plugins/Media/ImgMedia/Source/ImgMedia/Private/Loader/ImgMediaLoader.h
@@ -219,6 +219,11 @@ public:
 		return NumTilesY;
 	}
 
+	bool IsTiled() const
+	{
+		return (NumTilesX > 1) || (NumTilesY > 1);
+	}
+
 	/**
 	 * Get the next work item.
 	 *
