ImgMedia: Added log message (instead of crashing) when the CPU reader cannot read a file.
#rnx

diff --git a/Engine/Plugins/Media/ImgMedia/Source/ImgMedia/Private/Readers/ExrImgMediaReader.cpp b/Engine/Plugins/Media/ImgMedia/Source/ImgMedia/Private/Readers/ExrImgMediaReader.cpp
--- a/Engine/Plugins/Media/ImgMedia/Source/ImgMedia/Private/Readers/ExrImgMediaReader.cpp
+++ b/Engine/Plugins/Media/ImgMedia/Source/ImgMedia/Private/Readers/ExrImgMediaReader.cpp
@@ -173,14 +173,20 @@ bool FExrImgMediaReader::ReadFrame(int32 FrameId, int32 MipLevel, const FImgMedi
 						Image = FString::Printf(TEXT("%s_x%d_y%d.exr"), *BaseImage, TileX, TileY);
 					}
 					FRgbaInputFile InputFile(Image, 2);
+					if (InputFile.HasInputFile())
+					{
+						// read frame data
+						InputFile.SetFrameBuffer(MipDataPtr + FrameBufferOffsetX + FrameBufferOffsetY, Dim);
+						InputFile.ReadPixels(0, TileHeight - 1);
 
-					// read frame data
-					InputFile.SetFrameBuffer(MipDataPtr + FrameBufferOffsetX + FrameBufferOffsetY, Dim);
-					InputFile.ReadPixels(0, TileHeight - 1);
-
-					FrameBufferOffsetX += FrameBufferOffsetPerTileX;
-					OutFrame->MipMapsPresent |= 1 << CurrentMipLevel;
-					LevelFoundSoFar = true;
+						FrameBufferOffsetX += FrameBufferOffsetPerTileX;
+						OutFrame->MipMapsPresent |= 1 << CurrentMipLevel;
+						LevelFoundSoFar = true;
+					}
+					else
+					{
+						UE_LOG(LogImgMedia, Error, TEXT("Could not load %s"), *Image);
+					}
 				}
 				
 				FrameBufferOffsetY += FrameBufferOffsetPerTileY;
