ImgMedia: Added cache tab to the editor.
#rnx

diff --git a/Engine/Plugins/Media/ImgMedia/Source/ImgMediaEditor/ImgMediaEditor.Build.cs b/Engine/Plugins/Media/ImgMedia/Source/ImgMediaEditor/ImgMediaEditor.Build.cs
--- a/Engine/Plugins/Media/ImgMedia/Source/ImgMediaEditor/ImgMediaEditor.Build.cs
+++ b/Engine/Plugins/Media/ImgMedia/Source/ImgMediaEditor/ImgMediaEditor.Build.cs
@@ -21,6 +21,7 @@ namespace UnrealBuildTool.Rules
 					"Slate",
 					"SlateCore",
 					"UnrealEd",
+					"WorkspaceMenuStructure",
 				});
 
 			PrivateIncludePaths.AddRange(
diff --git a/Engine/Plugins/Media/ImgMedia/Source/ImgMediaEditor/Private/ImgMediaEditorModule.cpp b/Engine/Plugins/Media/ImgMedia/Source/ImgMediaEditor/Private/ImgMediaEditorModule.cpp
--- a/Engine/Plugins/Media/ImgMedia/Source/ImgMediaEditor/Private/ImgMediaEditorModule.cpp
+++ b/Engine/Plugins/Media/ImgMedia/Source/ImgMediaEditor/Private/ImgMediaEditorModule.cpp
@@ -12,9 +12,17 @@
 #include "ImgMediaSource.h"
 #include "PropertyEditorModule.h"
 #include "UObject/NameTypes.h"
+#include "Widgets/Docking/SDockTab.h"
+#include "Widgets/SImgMediaCache.h"
+#include "WorkspaceMenuStructure.h"
+#include "WorkspaceMenuStructureModule.h"
+
+#define LOCTEXT_NAMESPACE "ImgMediaEditorModule"
 
 DEFINE_LOG_CATEGORY(LogImgMediaEditor);
 
+static const FName ImgMediaCacheTabName(TEXT("ImgMediaCache"));
+
 /**
  * Implements the ImgMediaEditor module.
  */
@@ -29,10 +37,12 @@ public:
 	{
 		RegisterCustomizations();
 		RegisterAssetTools();
+		RegisterTabSpawners();
 	}
 
 	virtual void ShutdownModule() override
 	{
+		UnregisterTabSpawners();
 		UnregisterAssetTools();
 		UnregisterCustomizations();
 	}
@@ -90,6 +100,41 @@ protected:
 		}
 	}
 
+	void RegisterTabSpawners()
+	{
+		// Add ImgMedia group.
+		const IWorkspaceMenuStructure& MenuStructure = WorkspaceMenu::GetMenuStructure();
+		TSharedRef<FWorkspaceItem> MediaBrowserGroup = MenuStructure.GetLevelEditorCategory()->AddGroup(
+			LOCTEXT("WorkspaceMenu_ImgMediaCategory", "ImgMedia"),
+			FSlateIcon(FEditorStyle::GetStyleSetName(), "SequenceRecorder.TabIcon"),
+			true);
+
+		// Add cache tab.
+		FGlobalTabmanager::Get()->RegisterNomadTabSpawner(ImgMediaCacheTabName,
+			FOnSpawnTab::CreateStatic(&FImgMediaEditorModule::SpawnCacheTab))
+			.SetGroup(MediaBrowserGroup)
+			.SetDisplayName(LOCTEXT("ImgMediaCacheTabTitle", "Cache"))
+			.SetTooltipText(LOCTEXT("ImgMediaCacheTooltipText", "Open the cache tab."))
+			.SetIcon(FSlateIcon(FEditorStyle::GetStyleSetName(), "SequenceRecorder.TabIcon"));
+	}
+
+	void UnregisterTabSpawners()
+	{
+		if (FSlateApplication::IsInitialized())
+		{
+			FGlobalTabmanager::Get()->UnregisterNomadTabSpawner(ImgMediaCacheTabName);
+		}
+	}
+
+	static TSharedRef<SDockTab> SpawnCacheTab(const FSpawnTabArgs& SpawnTabArgs)
+	{
+		return SNew(SDockTab)
+			.TabRole(ETabRole::NomadTab)
+			[
+				SNew(SImgMediaCache)
+			];
+	}
+
 private:
 
 	/** Customization name to avoid reusing staticstruct during shutdown. */
@@ -102,3 +147,5 @@ private:
 
 
 IMPLEMENT_MODULE(FImgMediaEditorModule, ImgMediaEditor);
+
+#undef LOCTEXT_NAMESPACE
diff --git a/Engine/Plugins/Media/ImgMedia/Source/ImgMediaEditor/Private/Widgets/SImgMediaCache.cpp b/Engine/Plugins/Media/ImgMedia/Source/ImgMediaEditor/Private/Widgets/SImgMediaCache.cpp
new file mode 100644
--- /dev/null
+++ b/Engine/Plugins/Media/ImgMedia/Source/ImgMediaEditor/Private/Widgets/SImgMediaCache.cpp
@@ -0,0 +1,67 @@
+// Copyright Epic Games, Inc. All Rights Reserved.
+
+#include "Widgets/SImgMediaCache.h"
+#include "DetailLayoutBuilder.h"
+#include "IImgMediaModule.h"
+#include "ImgMediaGlobalCache.h"
+#include "SlateOptMacros.h"
+#include "Widgets/Input/SButton.h"
+#include "Widgets/Layout/SScrollBox.h"
+#include "Widgets/Layout/SSplitter.h"
+#include "Widgets/Text/STextBlock.h"
+
+#define LOCTEXT_NAMESPACE "ImgMediaCache"
+
+SImgMediaCache::~SImgMediaCache()
+{
+}
+
+BEGIN_SLATE_FUNCTION_BUILD_OPTIMIZATION
+void SImgMediaCache::Construct(const FArguments& InArgs)
+{
+	ChildSlot
+	[
+		SNew(SVerticalBox)
+
+		+ SVerticalBox::Slot()
+			.Padding(0, 20, 0, 0)
+			.AutoHeight()
+		
+		+ SVerticalBox::Slot()
+			.AutoHeight()
+			.Padding(0, 5, 0, 0)
+		[
+			SNew(STextBlock)
+				.Text(LOCTEXT("GlobalCache", "Global Cache"))
+				.Font(IDetailLayoutBuilder::GetDetailFont())
+		]
+			
+		+ SVerticalBox::Slot()
+			.AutoHeight()
+			.Padding(0, 5, 0, 0)
+			[
+				SNew(SButton)
+					.ButtonStyle(FEditorStyle::Get(), "HoverHintOnly")
+					.ForegroundColor(FSlateColor::UseForeground())
+					.OnClicked(this, &SImgMediaCache::OnClearGlobalCacheClicked)
+					.Text(LOCTEXT("ClearGlobalCache", "Clear Global Cache"))
+					.ToolTipText(LOCTEXT("ClearGlobalCacheButtonToolTip", "Clear the global cache."))
+			]
+		
+	];
+}
+END_SLATE_FUNCTION_BUILD_OPTIMIZATION
+
+FReply SImgMediaCache::OnClearGlobalCacheClicked()
+{
+	// Clear the cache.
+	FImgMediaGlobalCache* GlobalCache = IImgMediaModule::GetGlobalCache();
+	if (GlobalCache != nullptr)
+	{
+		GlobalCache->EmptyCache();
+	}
+
+	return FReply::Handled();
+}
+
+#undef LOCTEXT_NAMESPACE
diff --git a/Engine/Plugins/Media/ImgMedia/Source/ImgMediaEditor/Private/Widgets/SImgMediaCache.h b/Engine/Plugins/Media/ImgMedia/Source/ImgMediaEditor/Private/Widgets/SImgMediaCache.h
new file mode 100644
--- /dev/null
+++ b/Engine/Plugins/Media/ImgMedia/Source/ImgMediaEditor/Private/Widgets/SImgMediaCache.h
@@ -0,0 +1,26 @@
+// Copyright Epic Games, Inc. All Rights Reserved.
+
+#pragma once
+
+#include "CoreMinimal.h"
+#include "Input/Reply.h"
+#include "Widgets/SCompoundWidget.h"
+
+/**
+ * SImgMediaCache manages caching for image sequences.
+ */
+class SImgMediaCache : public SCompoundWidget
+{
+public:
+	SLATE_BEGIN_ARGS(SImgMediaCache){}
+	SLATE_END_ARGS()
+
+	virtual ~SImgMediaCache();
+
+	void Construct(const FArguments& InArgs);
+
+private:
+	/** Called when we click on the clear global cache button. */
+	FReply OnClearGlobalCacheClicked();
+
+};
