ImgMedia: Removed MipMapInfo customization as its obsolete now.
#rnx

diff --git a/Engine/Plugins/Media/ImgMedia/Source/ImgMediaEditor/Private/Customizations/ImgMediaSourceCustomization.cpp b/Engine/Plugins/Media/ImgMedia/Source/ImgMediaEditor/Private/Customizations/ImgMediaSourceCustomization.cpp
--- a/Engine/Plugins/Media/ImgMedia/Source/ImgMediaEditor/Private/Customizations/ImgMediaSourceCustomization.cpp
+++ b/Engine/Plugins/Media/ImgMedia/Source/ImgMediaEditor/Private/Customizations/ImgMediaSourceCustomization.cpp
@@ -96,193 +96,6 @@ void FImgMediaSourceCustomization::CustomizeChildren(TSharedRef<IPropertyHandle>
 {
 }
 
-void FImgMediaSourceCustomization::CustomizeMipMapInfo(IDetailLayoutBuilder& DetailBuilder)
-{
-	// Mipmap inspection only works for one object.
-	TArray<TWeakObjectPtr<UObject>> CustomizedObjects;
-	DetailBuilder.GetObjectsBeingCustomized(CustomizedObjects);
-	if (CustomizedObjects.Num() == 1)
-	{
-		UImgMediaSource* ImgMediaSource = Cast<UImgMediaSource>(CustomizedObjects[0].Get());
-		if (ImgMediaSource != nullptr)
-		{
-			const FImgMediaMipMapInfo* MipMapInfo = ImgMediaSource->GetMipMapInfo();
-			if (MipMapInfo != nullptr)
-			{
-				IDetailCategoryBuilder& MipMapInfoCategory = DetailBuilder.EditCategory("MipMapInfo", LOCTEXT("MipMapInfoCategoryName", "MipMapInfo"));
-
-				// Add objects.
-				AddMipMapObjects(MipMapInfoCategory, MipMapInfo);
-
-				// Loop over all cameras.
-				const TArray<FImgMediaMipMapCameraInfo>& CameraInfos = MipMapInfo->GetCameraInfo();
-				const TArray<float>& MipLevelDistances = MipMapInfo->GetMipLevelDistances();
-				float ViewportDistAdjust = MipMapInfo->GetViewportDistAdjust();
-				const TArray<FImgMediaMipMapObjectInfo*>& Objects = MipMapInfo->GetObjects();
-				if (CameraInfos.Num() > 0)
-				{
-					for (const FImgMediaMipMapCameraInfo& CameraInfo : CameraInfos)
-					{
-						// Add group for this camera.
-						IDetailGroup& Group = MipMapInfoCategory.AddGroup(*CameraInfo.Name, FText::FromString(CameraInfo.Name));
-
-						// Get mip level distances adjusted so they are in real space.
-						TArray<float> AdjustedMipLevelDistances;
-						for (float Dist : MipLevelDistances)
-						{
-							if (CameraInfo.ScreenSize == 0.0f)
-							{
-								Dist /= ViewportDistAdjust;
-							}
-							Dist /= CameraInfo.DistAdjust;
-							AdjustedMipLevelDistances.Add(Dist);
-						}
-
-						// Add info to UI.
-						AddCameraObjects(Group, CameraInfo, AdjustedMipLevelDistances, Objects);
-						AddCameraMipDistances(Group, CameraInfo, AdjustedMipLevelDistances);
-					}
-				}
-				else
-				{
-					FDetailWidgetRow& NoCamerasRow = MipMapInfoCategory.AddCustomRow(FText::GetEmpty());
-					NoCamerasRow.NameContent()
-						[
-							SNew(STextBlock)
-								.Text(LOCTEXT("MipMapNoCameras", "No Cameras"))
-						];
-				}
-			}
-		}
-	}
-}
-
-void FImgMediaSourceCustomization::AddMipMapObjects(IDetailCategoryBuilder& InCategory, const FImgMediaMipMapInfo* MipMapInfo)
-{
-	IDetailGroup& MainObjGroup = InCategory.AddGroup("MainObjects", LOCTEXT("MipMapObjects", "Objects"));
-
-	// Loop over all objects.
-	const TArray<FImgMediaMipMapObjectInfo*>& Objects = MipMapInfo->GetObjects();
-	for (const FImgMediaMipMapObjectInfo* ObjectInfo : Objects)
-	{
-		AActor* Object = ObjectInfo->Object.Get();
-		if (Object != nullptr)
-		{
-			IDetailGroup& ObjGroup = MainObjGroup.AddGroup(*Object->GetName(), FText::FromString(Object->GetName()));
-
-			// Show auto width.
-			FDetailWidgetRow& AutoWidthRow = ObjGroup.AddWidgetRow();
-			AutoWidthRow.NameContent()
-				[
-					SNew(STextBlock)
-						.Text(LOCTEXT("MipMapAutoWidth", "AutoWidth"))
-						.ToolTipText(LOCTEXT("MipMapAutoWidthToolTip", "This will be used for the width of the object if not set manually."))
-				];
-			AutoWidthRow.ValueContent()
-				[
-					SNew(STextBlock)
-						.Text(FText::FromString(FString::Printf(TEXT("%f"), FImgMediaMipMapInfo::GetObjectWidth(Object))))
-						.ToolTipText(LOCTEXT("MipMapAutoWidthToolTip", "This will be used for the width of the object if not set manually."))
-				];
-
-			// Show width.
-			FDetailWidgetRow& WidthRow = ObjGroup.AddWidgetRow();
-			WidthRow.NameContent()
-				[
-					SNew(STextBlock)
-						.Text(LOCTEXT("MipMapWidth", "Width"))
-						.ToolTipText(LOCTEXT("MipMapWidthToolTip", "Width of this object that is currently used for mipmap calculations."))
-				];
-			WidthRow.ValueContent()
-				[
-					SNew(STextBlock)
-						.Text(FText::FromString(FString::Printf(TEXT("%f"), ObjectInfo->Width)))
-						.ToolTipText(LOCTEXT("MipMapWidthToolTip", "Width of this object that is currently used for mipmap calculations."))
-				];
-		}
-	}
-}
-
-
-void FImgMediaSourceCustomization::AddCameraObjects(IDetailGroup& InCameraGroup, const FImgMediaMipMapCameraInfo& InCameraInfo, const TArray<float>& InMipLevelDistances, const TArray<FImgMediaMipMapObjectInfo*>& InObjects)
-{
-	IDetailGroup& Group = InCameraGroup.AddGroup("Objects", LOCTEXT("MipMapObjects", "Objects"));
-
-	// Loop over all objects.
-	for (const FImgMediaMipMapObjectInfo* ObjectInfo : InObjects)
-	{
-		AActor* Object = ObjectInfo->Object.Get();
-		if (Object != nullptr)
-		{
-			IDetailGroup& ObjGroup = Group.AddGroup(*Object->GetName(), FText::FromString(Object->GetName()));
-
-			// Show distance to camera.
-			float DistToCamera = FImgMediaMipMapInfo::GetObjectDistToCamera(InCameraInfo.Location, Object->GetActorLocation());
-			FDetailWidgetRow& DistToCameraRow = ObjGroup.AddWidgetRow();
-			DistToCameraRow.NameContent()
-				[
-					SNew(STextBlock)
-						.Text(LOCTEXT("MipMapObjDistToCamera", "Distance To Camera"))
-				];
-			DistToCameraRow.ValueContent()
-				[
-					SNew(STextBlock)
-						.Text(FText::FromString(FString::Printf(TEXT("%f"), DistToCamera)))
-				];
-
-			// Show adjusted distance to camera.
-			DistToCamera -= ObjectInfo->Width;
-			DistToCamera *= ObjectInfo->DistAdjust; 
-			FDetailWidgetRow& AdjustedDistToCameraRow = ObjGroup.AddWidgetRow();
-			AdjustedDistToCameraRow.NameContent()
-				[
-					SNew(STextBlock)
-						.Text(LOCTEXT("MipMapObjAdjustedDistToCamera", "Adjusted Distance To Camera"))
-				];
-			AdjustedDistToCameraRow.ValueContent()
-				[
-					SNew(STextBlock)
-						.Text(FText::FromString(FString::Printf(TEXT("%f"), DistToCamera)))
-				];
-
-			// Show mip level.
-			int MipLevel = FImgMediaMipMapInfo::GetMipLevelForDistance(DistToCamera, InMipLevelDistances);
-			FDetailWidgetRow& MipLevelRow = ObjGroup.AddWidgetRow();
-			MipLevelRow.NameContent()
-				[
-					SNew(STextBlock)
-						.Text(LOCTEXT("MipMapObjMipLevel", "Mip Level"))
-				];
-			MipLevelRow.ValueContent()
-				[
-					SNew(STextBlock)
-						.Text(FText::FromString(FString::Printf(TEXT("%d"), MipLevel)))
-				];
-		}
-	}
-}
-
-void FImgMediaSourceCustomization::AddCameraMipDistances(IDetailGroup& InCameraGroup, const FImgMediaMipMapCameraInfo& InCameraInfo, const TArray<float>& InMipLevelDistances)
-{
-	IDetailGroup& Group = InCameraGroup.AddGroup("MipLevelDistances", LOCTEXT("MipMapLevelDistances", "MipLevelDistances"));
-
-	for (int Index = 0; Index < InMipLevelDistances.Num(); ++Index)
-	{
-		float Distance = InMipLevelDistances[Index];
-		FDetailWidgetRow& DistancesRow = Group.AddWidgetRow();
-		DistancesRow.NameContent()
-			[
-				SNew(STextBlock)
-					.Text(FText::FromString(FString::Printf(TEXT("%d"), Index)))
-			];
-		DistancesRow.ValueContent()
-			[
-				SNew(STextBlock)
-					.Text(FText::FromString(FString::Printf(TEXT("%f"), Distance)))
-			];
-	}
-}
-
 /* FImgMediaSourceCustomization implementation
  *****************************************************************************/
 
diff --git a/Engine/Plugins/Media/ImgMedia/Source/ImgMediaEditor/Private/Customizations/ImgMediaSourceCustomization.h b/Engine/Plugins/Media/ImgMedia/Source/ImgMediaEditor/Private/Customizations/ImgMediaSourceCustomization.h
--- a/Engine/Plugins/Media/ImgMedia/Source/ImgMediaEditor/Private/Customizations/ImgMediaSourceCustomization.h
+++ b/Engine/Plugins/Media/ImgMedia/Source/ImgMediaEditor/Private/Customizations/ImgMediaSourceCustomization.h
@@ -53,40 +53,6 @@ public:
 
 protected:
 
-	/**
-	 * Adds custom UI for mipmap info.
-	 *
-	 * @param DetailBuilder Adds to this UI.
-	 */
-	void CustomizeMipMapInfo(IDetailLayoutBuilder& DetailBuilder);
-
-	/**
-	 * Adds info on objects to the UI.
-	 *
-	 * @param InCategory UI will be added here.
-	 * @param MipMapInfo Object info will be retrieved from this.
-	 */
-	void AddMipMapObjects(IDetailCategoryBuilder& InCategory, const FImgMediaMipMapInfo* MipMapInfo);
-
-	/**
-	 * Adds object info for a camera to the UI.
-	 * 
-	 * @param InCameraGroup UI to add to.
-	 * @param InCameraInfo Camera to use.
-	 * @param InMipLevelDistances Distances for each mip map level in world space.
-	 * @param Objects List of objects to show info for.
-	 */
-	void AddCameraObjects(IDetailGroup& InCameraGroup, const FImgMediaMipMapCameraInfo& InCameraInfo, const TArray<float>& InMipLevelDistances, const TArray<FImgMediaMipMapObjectInfo*>& Objects);
-	
-	/**
-	 * Adds mip level distances for a camera to the UI.
-	 *
-	 * @param InCameraGroup UI to add to.
-	 * @param InCameraInfo Camera to use.
-	 * @param InMipLevelDistances Distances for each mip map level in world space.
-	 */
-	void AddCameraMipDistances(IDetailGroup& InCameraGroup, const FImgMediaMipMapCameraInfo& InCameraInfo, const TArray<float>& InMipLevelDistances);
-
 	/**
 	 * Get the path to the currently selected image sequence.
 	 *
