ImgMedia: MipMapInfo uses GetActorBounds to get the size of the object.
#rnx

diff --git a/Engine/Plugins/Media/ImgMedia/Source/ImgMedia/Private/Assets/ImgMediaMipMapInfo.cpp b/Engine/Plugins/Media/ImgMedia/Source/ImgMedia/Private/Assets/ImgMediaMipMapInfo.cpp
--- a/Engine/Plugins/Media/ImgMedia/Source/ImgMedia/Private/Assets/ImgMediaMipMapInfo.cpp
+++ b/Engine/Plugins/Media/ImgMedia/Source/ImgMedia/Private/Assets/ImgMediaMipMapInfo.cpp
@@ -41,12 +41,8 @@ void FImgMediaMipMapInfo::AddObject(AActor* InActor, float Width, float LODBias)
 		// Calculate dist adjust from bias.
 		Info->DistAdjust = FMath::Pow(2.0f, LODBias);
 
-		// If width < 0, then try to calculate it.
-		if (Width < 0.0f)
-		{
-			Width = GetObjectWidth(InActor);
-		}
-		Info->Width = Width;
+		// Get size of object.
+		GetObjectSize(InActor, Info->Width, Info->Height);
 
 		// Get dist adjust.
 		if (Width > 0.0f)
@@ -276,23 +272,17 @@ int FImgMediaMipMapInfo::GetMipLevelForDistance(float InDistance, const TArray<f
 	return MipLevel;
 }
 
-float FImgMediaMipMapInfo::GetObjectWidth(const AActor* InActor)
+void FImgMediaMipMapInfo::GetObjectSize(const AActor* InActor, float& Width, float& Height)
 {
 	// Get box extent of actor to calculate width from.
-	float Width = -1.0f;
+	Width = -1.0f;
 	if (InActor != nullptr)
 	{
-		const USceneComponent* RootComponent = InActor->GetRootComponent();
-		if (RootComponent != nullptr)
-		{
-			const UStaticMeshComponent* MeshComponent = Cast<const UStaticMeshComponent>(RootComponent);
-			if (MeshComponent != nullptr)
-			{
-				const UStaticMesh* StaticMesh = MeshComponent->GetStaticMesh();
-				const FBoxSphereBounds Bounds = StaticMesh->GetBounds();
-				Width = Bounds.BoxExtent.X * 2.0f;
-			}
-		}
+		FVector BoundsOrigin;
+		FVector BoxExtent;
+		InActor->GetActorBounds(false, BoundsOrigin, BoxExtent);
+		Width = BoxExtent.Z * 2.0f;
+		Height = BoxExtent.Y * 2.0f;
 	}
 
 	// Did we get anything?
@@ -300,8 +290,7 @@ float FImgMediaMipMapInfo::GetObjectWidth(const AActor* InActor)
 	{
 		UE_LOG(LogImgMedia, Error, TEXT("FImgMediaMipMapInfo could not get size of %s."), InActor != nullptr ? *InActor->GetName() : TEXT("<nullptr>"));
 		Width = 0.0f;
+		Height = 0.0f;
 	}
-
-	return Width;
 }
 
diff --git a/Engine/Plugins/Media/ImgMedia/Source/ImgMedia/Public/ImgMediaMipMapInfo.h b/Engine/Plugins/Media/ImgMedia/Source/ImgMedia/Public/ImgMediaMipMapInfo.h
--- a/Engine/Plugins/Media/ImgMedia/Source/ImgMedia/Public/ImgMediaMipMapInfo.h
+++ b/Engine/Plugins/Media/ImgMedia/Source/ImgMedia/Public/ImgMediaMipMapInfo.h
@@ -40,8 +40,10 @@ struct FImgMediaMipMapObjectInfo
 {
 	/** Actor that is using our img sequence. */
 	TWeakObjectPtr<class AActor> Object;
-	/** Size of this object. */
+	/** Width of this object. */
 	float Width;
+	/** Height of this object. */
+	float Height;
 	/** LOD bias for the mipmap level. */
 	float LODBias;
 	/** Multiplier to apply to the distance to account for this object (e.g. its size). */
@@ -190,9 +192,10 @@ public:
 	 * Determine the size of an object.
 	 *
 	 * @param InActor Object to get the size of.
-	 * @return Size.
+	 * @param Width		Will be set to the width of the object.
+	 * @param Height	Will be set to the height of the object.
 	 */
-	IMGMEDIA_API static float GetObjectWidth(const AActor* InActor);
+	void GetObjectSize(const AActor* InActor, float& Width, float& Height);
 	
 	/**
 	 * Get information on all our cameras.
