ImgMedia: Uses different camera sources so we can get more information needed for tiling.

#rnx

diff --git a/Engine/Plugins/Media/ImgMedia/Source/ImgMedia/ImgMedia.Build.cs b/Engine/Plugins/Media/ImgMedia/Source/ImgMedia/ImgMedia.Build.cs
--- a/Engine/Plugins/Media/ImgMedia/Source/ImgMedia/ImgMedia.Build.cs
+++ b/Engine/Plugins/Media/ImgMedia/Source/ImgMedia/ImgMedia.Build.cs
@@ -82,6 +82,12 @@ namespace UnrealBuildTool.Rules
 					//required for FPostProcessMaterialInputs
 					Path.Combine(EngineDir, "Source/Runtime/Renderer/Private")
 				});
+
+			// Is this the editor?
+			if (Target.bBuildEditor == true)
+			{
+				PublicDependencyModuleNames.Add("UnrealEd");
+			}
 		}
 	}
 }
diff --git a/Engine/Plugins/Media/ImgMedia/Source/ImgMedia/Private/Assets/ImgMediaMipMapInfoManager.cpp b/Engine/Plugins/Media/ImgMedia/Source/ImgMedia/Private/Assets/ImgMediaMipMapInfoManager.cpp
--- a/Engine/Plugins/Media/ImgMedia/Source/ImgMedia/Private/Assets/ImgMediaMipMapInfoManager.cpp
+++ b/Engine/Plugins/Media/ImgMedia/Source/ImgMedia/Private/Assets/ImgMediaMipMapInfoManager.cpp
@@ -12,7 +12,15 @@
 #include "Camera/CameraComponent.h"
 #include "Engine/GameEngine.h"
 #include "Engine/GameViewportClient.h"
+#include "Engine/LocalPlayer.h"
+#include "GameFramework/PlayerController.h"
 #include "MovieSceneCommonHelpers.h"
+#include "SceneView.h"
+
+#if WITH_EDITOR
+#include "EditorViewportClient.h"
+#include "LevelEditorViewport.h"
+#endif
 
 static TAutoConsoleVariable<bool> CVarImgMediaMipMapDebugEnable(
 	TEXT("ImgMedia.MipMapDebug"),
@@ -149,42 +157,76 @@ void FImgMediaMipMapInfoManager::UpdateCameraInfo()
 {
 	CameraInfos.Empty();
 
-	// Get camera info from the engine.
-	FImgMediaEngine::Get().GetCameraInfo(CameraInfos);
-	// Take screen size into account if set.
-	for (FImgMediaMipMapCameraInfo& CameraInfo : CameraInfos)
+	const FString Name = "Unknown";
+
+	// Loop through players.
+	UGameViewportClient* GameViewportClient = GEngine->GameViewport;
+	if (GameViewportClient != nullptr)
 	{
-		if (CameraInfo.ScreenSize != 0.0f)
+		UWorld* World = GameViewportClient->GetWorld();
+		if (World != nullptr)
 		{
-			float ScreenAdjust = RefFrameBufferWidth / CameraInfo.ScreenSize;
-			CameraInfo.DistAdjust *= ScreenAdjust;
+			for (FConstPlayerControllerIterator Iterator = World->GetPlayerControllerIterator(); Iterator; ++Iterator)
+			{
+				APlayerController* PlayerController = Iterator->Get();
+				ULocalPlayer* LocalPlayer = PlayerController->GetLocalPlayer();
+				if (LocalPlayer != nullptr)
+				{
+					// Get view info.
+					FSceneViewFamilyContext ViewFamily(FSceneViewFamily::ConstructionValues(
+						LocalPlayer->ViewportClient->Viewport,
+						World->Scene,
+						LocalPlayer->ViewportClient->EngineShowFlags)
+						.SetRealtimeUpdate(true));
+					FVector ViewLocation;
+					FRotator ViewRotation;
+					FSceneView* SceneView = LocalPlayer->CalcSceneView(&ViewFamily, ViewLocation, ViewRotation, LocalPlayer->ViewportClient->Viewport);
+
+					// Did we get a view?
+					if (SceneView != nullptr)
+					{
+						// Add camera info.
+						float NearPlaneSize = CalculateNearPlaneSize(SceneView->FOV);
+						float DistAdjust = NearPlaneSize / RefNearPlaneWidth;
+						CameraInfos.Emplace(Name, ViewLocation, 0.0f, DistAdjust);
+					}
+				}
+			}
 		}
 	}
 
-	// Loop over all cameras.
-	for (TWeakObjectPtr<UCameraComponent> CameraPtr : Cameras)
+#if WITH_EDITOR
+	// If we are have no cameras then we could be in the editor and not PIE so get the editor view.
+	if (CameraInfos.IsEmpty())
 	{
-		UCameraComponent* CameraComponent = CameraPtr.Get();
-		if (CameraComponent != nullptr)
+		// Loop through all viewport clients.
+		for (FEditorViewportClient* ViewportClient : GEditor->GetAllViewportClients())
 		{
-			// Get info.
-			FString Name;
-			AActor* Owner = CameraComponent->GetOwner();
-			if (Owner != nullptr)
-			{
-				Name = Owner->GetName();
-			}
-			else
+			if (ViewportClient != nullptr) 
 			{
-				Name = "Unknown";
+				FViewport* Viewport = ViewportClient->Viewport;
+				if (Viewport != nullptr)
+				{
+					// The editor has a few viewports,
+					// and some that we do not care about have no size,
+					// so filter on this.
+					FIntPoint ViewportSize = Viewport->GetSizeXY();
+					if ((ViewportSize.X > 0) && (ViewportSize.Y > 0))
+					{
+						// Get the view.
+						const FViewportCameraTransform& ViewTransform = ViewportClient->GetViewTransform();
+						
+						// Add camera info.
+						float NearPlaneSize = CalculateNearPlaneSize(ViewportClient->ViewFOV);
+						float DistAdjust = NearPlaneSize / RefNearPlaneWidth;
+						DistAdjust *= RefFrameBufferWidth / ViewportSize.X;
+						CameraInfos.Emplace(Name, ViewTransform.GetLocation(), ViewportSize.X, DistAdjust);
+					}
+				}
 			}
-			FMinimalViewInfo CameraView;
-			CameraComponent->GetCameraView(FApp::GetDeltaTime(), CameraView);
-			float NearPlaneSize = CalculateNearPlaneSize(CameraView.FOV);
-			float DistAdjust = NearPlaneSize / RefNearPlaneWidth;
-			CameraInfos.Emplace(Name, CameraView.Location, 0.0f, DistAdjust);
 		}
 	}
+#endif
 }
 
 float FImgMediaMipMapInfoManager::CalculateNearPlaneSize(float InFOV) const
