ImgMedia: Fixed mipmap distance adjustment due to object orientation using the full width instead of half.
#rnx

diff --git a/Engine/Plugins/Media/ImgMedia/Source/ImgMedia/Private/Assets/ImgMediaMipMapInfo.cpp b/Engine/Plugins/Media/ImgMedia/Source/ImgMedia/Private/Assets/ImgMediaMipMapInfo.cpp
--- a/Engine/Plugins/Media/ImgMedia/Source/ImgMedia/Private/Assets/ImgMediaMipMapInfo.cpp
+++ b/Engine/Plugins/Media/ImgMedia/Source/ImgMedia/Private/Assets/ImgMediaMipMapInfo.cpp
@@ -173,7 +173,7 @@ void FImgMediaMipMapInfo::UpdateMipLevelCache()
 
 					// Account for object size and possible rotation.
 					// Closest case is the object edge is pointing towards us.
-					AdjustedDistToCamera = AdjustedDistToCamera - Info->Width;
+					AdjustedDistToCamera = AdjustedDistToCamera - Info->Width * 0.5f;
 
 					// If we are using our calculated distances,
 					// then account for various settings.
