ImgMedia: Fix for crash when running an editor build with -game.
#rnx

diff --git a/Engine/Plugins/Media/ImgMedia/Source/ImgMedia/Private/Assets/ImgMediaMipMapInfoManager.cpp b/Engine/Plugins/Media/ImgMedia/Source/ImgMedia/Private/Assets/ImgMediaMipMapInfoManager.cpp
--- a/Engine/Plugins/Media/ImgMedia/Source/ImgMedia/Private/Assets/ImgMediaMipMapInfoManager.cpp
+++ b/Engine/Plugins/Media/ImgMedia/Source/ImgMedia/Private/Assets/ImgMediaMipMapInfoManager.cpp
@@ -197,7 +197,7 @@ void FImgMediaMipMapInfoManager::UpdateCameraInfo()
 
 #if WITH_EDITOR
 	// If we are have no cameras then we could be in the editor and not PIE so get the editor view.
-	if (CameraInfos.IsEmpty())
+	if ((CameraInfos.IsEmpty()) && (GEditor != nullptr))
 	{
 		// Loop through all viewport clients.
 		for (FEditorViewportClient* ViewportClient : GEditor->GetAllViewportClients())
