ImgMedia: Removed ImgMediaEngine as its obsolete.
#rnx

diff --git a/Engine/Plugins/Media/ImgMedia/Source/ImgMedia/Private/Assets/ImgMediaMipMapInfoManager.cpp b/Engine/Plugins/Media/ImgMedia/Source/ImgMedia/Private/Assets/ImgMediaMipMapInfoManager.cpp
--- a/Engine/Plugins/Media/ImgMedia/Source/ImgMedia/Private/Assets/ImgMediaMipMapInfoManager.cpp
+++ b/Engine/Plugins/Media/ImgMedia/Source/ImgMedia/Private/Assets/ImgMediaMipMapInfoManager.cpp
@@ -1,7 +1,6 @@
 // Copyright Epic Games, Inc. All Rights Reserved.
 
 #include "ImgMediaMipMapInfoManager.h"
-#include "ImgMediaEngine.h"
 #include "ImgMediaPrivate.h"
 #include "MediaPlayer.h"
 #include "MediaPlayerFacade.h"
diff --git a/Engine/Plugins/Media/ImgMedia/Source/ImgMediaEngine/Private/ImgMediaEngine.cpp b/Engine/Plugins/Media/ImgMedia/Source/ImgMediaEngine/Private/ImgMediaEngine.cpp
deleted file mode 100644
diff --git a/Engine/Plugins/Media/ImgMedia/Source/ImgMediaEngine/Private/Unreal/ImgMediaEngineUnreal.cpp b/Engine/Plugins/Media/ImgMedia/Source/ImgMediaEngine/Private/Unreal/ImgMediaEngineUnreal.cpp
deleted file mode 100644
diff --git a/Engine/Plugins/Media/ImgMedia/Source/ImgMediaEngine/Private/Unreal/ImgMediaEngineUnreal.h b/Engine/Plugins/Media/ImgMedia/Source/ImgMediaEngine/Private/Unreal/ImgMediaEngineUnreal.h
deleted file mode 100644
diff --git a/Engine/Plugins/Media/ImgMedia/Source/ImgMediaEngine/Public/ImgMediaEngine.h b/Engine/Plugins/Media/ImgMedia/Source/ImgMediaEngine/Public/ImgMediaEngine.h
deleted file mode 100644
