ImgMedia: Fixed "clear global cache" button not looking like a button.
#rnx

diff --git a/Engine/Plugins/Media/ImgMedia/Source/ImgMediaEditor/Private/Widgets/SImgMediaCache.cpp b/Engine/Plugins/Media/ImgMedia/Source/ImgMediaEditor/Private/Widgets/SImgMediaCache.cpp
--- a/Engine/Plugins/Media/ImgMedia/Source/ImgMediaEditor/Private/Widgets/SImgMediaCache.cpp
+++ b/Engine/Plugins/Media/ImgMedia/Source/ImgMediaEditor/Private/Widgets/SImgMediaCache.cpp
@@ -38,11 +38,10 @@ void SImgMediaCache::Construct(const FArguments& InArgs)
 			
 		+ SVerticalBox::Slot()
 			.AutoHeight()
-			.Padding(0, 5, 0, 0)
+			.Padding(4.0f)
+			.HAlign(HAlign_Left)
 			[
 				SNew(SButton)
-					.ButtonStyle(FEditorStyle::Get(), "HoverHintOnly")
-					.ForegroundColor(FSlateColor::UseForeground())
 					.OnClicked(this, &SImgMediaCache::OnClearGlobalCacheClicked)
 					.Text(LOCTEXT("ClearGlobalCache", "Clear Global Cache"))
 					.ToolTipText(LOCTEXT("ClearGlobalCacheButtonToolTip", "Clear the global cache."))
