ImgMedia: Visibility of tiles is now computed.

#rnx

diff --git a/Engine/Plugins/Media/ImgMedia/Source/ImgMedia/Private/Assets/ImgMediaMipMapInfo.cpp b/Engine/Plugins/Media/ImgMedia/Source/ImgMedia/Private/Assets/ImgMediaMipMapInfo.cpp
--- a/Engine/Plugins/Media/ImgMedia/Source/ImgMedia/Private/Assets/ImgMediaMipMapInfo.cpp
+++ b/Engine/Plugins/Media/ImgMedia/Source/ImgMedia/Private/Assets/ImgMediaMipMapInfo.cpp
@@ -5,12 +5,18 @@
 #include "ImgMediaMipMapInfoManager.h"
 #include "ImgMediaPrivate.h"
 
+#include "Async/Async.h"
 #include "Components/StaticMeshComponent.h"
 #include "Engine/Engine.h"
 #include "Engine/StaticMesh.h"
 #include "GameFramework/Actor.h"
+#include "Math/UnrealMathUtility.h"
 #include "MediaTextureTracker.h"
 
+#if WITH_EDITOR
+#include "Editor.h"
+#endif
+
 DECLARE_CYCLE_STAT(TEXT("ImgMedia MipMap Update Cache"), STAT_ImgMedia_MipMapUpdateCache, STATGROUP_Media);
 
 FImgMediaMipMapInfo::FImgMediaMipMapInfo()
@@ -107,10 +113,13 @@ void FImgMediaMipMapInfo::SetMipLevelDistance(float Distance)
 	UpdateMipLevelDistances();
 }
 
-void FImgMediaMipMapInfo::SetTextureInfo(FName InSequenceName, int NumMipMaps, const FIntPoint& Dim)
+void FImgMediaMipMapInfo::SetTextureInfo(FName InSequenceName, int32 NumMipMaps,
+	int32 InNumTilesX, int32 InNumTilesY, const FIntPoint& Dim)
 {
 	SequenceName = InSequenceName;
 	MipLevelDistances.SetNum(NumMipMaps);
+	NumTilesX = InNumTilesX;
+	NumTilesY = InNumTilesY;
 
 	// Set ideal distance for mip level 0 based on various factors.
 	if (bIsMipLevel0DistanceSetManually == false)
@@ -153,6 +162,7 @@ void FImgMediaMipMapInfo::UpdateMipLevelCache()
 
 	CachedMipLevel = 0;
 	CachedTileSelection.SetAllVisible();
+	bool bHasTiles = (NumTilesX > 1) || (NumTilesY > 1);
 
 	// Loop over all objects.
 	float ClosestDistToCamera = TNumericLimits<float>::Max();
@@ -164,10 +174,37 @@ void FImgMediaMipMapInfo::UpdateMipLevelCache()
 			AActor* Object = Info->Object.Get();
 			if (Object != nullptr)
 			{
-				// Get closest distance to camera.
 				FVector ObjectLocation = Object->GetActorLocation();
+				
+				const FTransform ObjectTransform = Object->GetTransform();
+				FVector ObjectScale = Object->GetActorScale();
+				
+				// Set up for tiles.
+				float TileWidthInWorldSpace = 0.0f;
+				float TileHeightInWorldSpace = 0.0f;
+				float TileRadiusInWorldSpace = 0.0f;
+				FVector NextTileXVector;
+				FVector NextTileYVector;
+				if (bHasTiles)
+				{
+					// Get size of tile in world space.
+					TileWidthInWorldSpace = ObjectScale.X * Info->Width / ((float)NumTilesX);
+					TileHeightInWorldSpace = ObjectScale.Y * Info->Height / ((float)NumTilesY);
+					TileRadiusInWorldSpace = FMath::Max(TileWidthInWorldSpace, TileHeightInWorldSpace) * 0.5f;
+
+					// Get vectors to go from one tile to the next in world space.
+					FVector TilePosInObjectSpace(0.0f, TileWidthInWorldSpace, 0.0f);
+					FVector TilePos = ObjectTransform.TransformPositionNoScale(TilePosInObjectSpace);
+					NextTileXVector = TilePos - ObjectLocation;
+					TilePosInObjectSpace.Set(0.0f, 0.0f, TileHeightInWorldSpace);
+					TilePos = ObjectTransform.TransformPositionNoScale(TilePosInObjectSpace);
+					NextTileYVector = TilePos - ObjectLocation;
+				}
+
 				for (const FImgMediaMipMapCameraInfo& CameraInfo : CameraInfos)
 				{
+					// Mips calculation.
+					// Get closest distance to camera.
 					float DistToCamera = GetObjectDistToCamera(CameraInfo.Location, ObjectLocation);
 					float AdjustedDistToCamera = DistToCamera;
 
@@ -191,6 +228,13 @@ void FImgMediaMipMapInfo::UpdateMipLevelCache()
 					{
 						ClosestDistToCamera = AdjustedDistToCamera;
 					}
+
+					// See which tiles are visible.
+					if (bHasTiles)
+					{
+						CalculateTileVisibility(CameraInfo, ObjectLocation, NextTileXVector,
+							NextTileYVector, TileRadiusInWorldSpace);
+					}
 				}
 			}
 		}
@@ -207,6 +251,74 @@ void FImgMediaMipMapInfo::UpdateMipLevelCache()
 	bIsCachedMipLevelValid = true;
 }
 
+void FImgMediaMipMapInfo::CalculateTileVisibility(const FImgMediaMipMapCameraInfo& CameraInfo,
+	const FVector& ObjectLocation, const FVector& NextTileXVector,
+	const FVector& NextTileYVector, float TileRadiusInWorldSpace)
+{
+	// Get frustum.
+	FConvexVolume ViewFrustum;
+	GetViewFrustumBounds(ViewFrustum, CameraInfo.ViewMatrix, false);
+
+	// This is the middle of the object.
+	FVector TileStartPos = ObjectLocation;
+
+	// Loop over Y.
+	for (int32 TileY = 0; TileY < NumTilesY; ++TileY)
+	{
+		// Y needs to be flipped so that 0 is at the top.
+		// - NumTiles/2 so we start at the corner.
+		// + 0.5 so the position is the middle of the tile.
+		float TileMultY = (NumTilesY - TileY - 1) - ((float)NumTilesY) / 2.0f + 0.5f;
+
+		// Loop over X.
+		for (int32 TileX = 0; TileX < NumTilesX; ++TileX)
+		{
+			float TileMultX = TileX - ((float)NumTilesX) / 2.0f + 0.5f;
+			FVector TilePos = TileStartPos + NextTileXVector * TileMultX +
+				NextTileYVector * TileMultY;
+
+			// Is this tile visible?
+			if (ViewFrustum.IntersectSphere(TilePos, TileRadiusInWorldSpace))
+			{
+				// Is the cache initialized to all visible?
+				if (CachedTileSelection.BottomRightX != 0xffff)
+				{
+					// Nope.
+					// Add this tile to the selection.
+					CachedTileSelection.TopLeftX = FMath::Min(
+						CachedTileSelection.TopLeftX, (uint16)TileX);
+					CachedTileSelection.TopLeftY = FMath::Min(
+						CachedTileSelection.TopLeftY, (uint16)TileY);
+					CachedTileSelection.BottomRightX = FMath::Max(
+						CachedTileSelection.BottomRightX, (uint16)(TileX + 1));
+					CachedTileSelection.BottomRightY = FMath::Max(
+						CachedTileSelection.BottomRightY, (uint16)(TileY + 1));
+				}
+				else
+				{
+					// Yes.
+					// Set tile selection to this tile.
+					CachedTileSelection.TopLeftX = TileX;
+					CachedTileSelection.TopLeftY = TileY;
+					CachedTileSelection.BottomRightX = TileX + 1;
+					CachedTileSelection.BottomRightY = TileY + 1;
+				}
+				
+				// Enable this to draw a sphere where each tile is.
+#if false
+#if WITH_EDITOR
+				Async(EAsyncExecution::TaskGraphMainThread, [TilePos, TileRadiusInWorldSpace]()
+				{
+					UWorld* World = GEditor->GetEditorWorldContext().World();
+					DrawDebugSphere(World, TilePos, TileRadiusInWorldSpace, 6, FColor::Red);
+				});
+#endif // WITH_EDITOR
+#endif // false
+			}
+		}
+	}
+}
+
 void FImgMediaMipMapInfo::Tick(float DeltaTime)
 {
 	FScopeLock Lock(&InfoCriticalSection);
diff --git a/Engine/Plugins/Media/ImgMedia/Source/ImgMedia/Private/Assets/ImgMediaMipMapInfoManager.cpp b/Engine/Plugins/Media/ImgMedia/Source/ImgMedia/Private/Assets/ImgMediaMipMapInfoManager.cpp
--- a/Engine/Plugins/Media/ImgMedia/Source/ImgMedia/Private/Assets/ImgMediaMipMapInfoManager.cpp
+++ b/Engine/Plugins/Media/ImgMedia/Source/ImgMedia/Private/Assets/ImgMediaMipMapInfoManager.cpp
@@ -156,8 +156,6 @@ void FImgMediaMipMapInfoManager::UpdateCameraInfo()
 {
 	CameraInfos.Empty();
 
-	const FString Name = "Unknown";
-
 	// Loop through players.
 	UGameViewportClient* GameViewportClient = GEngine->GameViewport;
 	if (GameViewportClient != nullptr)
@@ -171,7 +169,7 @@ void FImgMediaMipMapInfoManager::UpdateCameraInfo()
 				ULocalPlayer* LocalPlayer = PlayerController->GetLocalPlayer();
 				if (LocalPlayer != nullptr)
 				{
-					// Get view info.
+					// Get the scene view.
 					FSceneViewFamilyContext ViewFamily(FSceneViewFamily::ConstructionValues(
 						LocalPlayer->ViewportClient->Viewport,
 						World->Scene,
@@ -184,10 +182,8 @@ void FImgMediaMipMapInfoManager::UpdateCameraInfo()
 					// Did we get a view?
 					if (SceneView != nullptr)
 					{
-						// Add camera info.
-						float NearPlaneSize = CalculateNearPlaneSize(SceneView->FOV);
-						float DistAdjust = NearPlaneSize / RefNearPlaneWidth;
-						CameraInfos.Emplace(Name, ViewLocation, 0.0f, DistAdjust);
+						// Add camera info for this scene view.
+						AddCameraInfo(SceneView);
 					}
 				}
 			}
@@ -212,14 +208,16 @@ void FImgMediaMipMapInfoManager::UpdateCameraInfo()
 					FIntPoint ViewportSize = Viewport->GetSizeXY();
 					if ((ViewportSize.X > 0) && (ViewportSize.Y > 0))
 					{
-						// Get the view.
-						const FViewportCameraTransform& ViewTransform = ViewportClient->GetViewTransform();
+						// Get the scene view.
+						FSceneViewFamilyContext ViewFamily(FSceneViewFamily::ConstructionValues(
+							Viewport,
+							ViewportClient->GetScene(),
+							ViewportClient->EngineShowFlags)
+							.SetRealtimeUpdate(ViewportClient->IsRealtime()));
+						FSceneView* SceneView = ViewportClient->CalcSceneView(&ViewFamily);
 						
-						// Add camera info.
-						float NearPlaneSize = CalculateNearPlaneSize(ViewportClient->ViewFOV);
-						float DistAdjust = NearPlaneSize / RefNearPlaneWidth;
-						DistAdjust *= RefFrameBufferWidth / ViewportSize.X;
-						CameraInfos.Emplace(Name, ViewTransform.GetLocation(), ViewportSize.X, DistAdjust);
+						// Add camera info for this scene view.
+						AddCameraInfo(SceneView);
 					}
 				}
 			}
@@ -228,6 +226,21 @@ void FImgMediaMipMapInfoManager::UpdateCameraInfo()
 #endif
 }
 
+void FImgMediaMipMapInfoManager::AddCameraInfo(FSceneView* SceneView)
+{
+	float FOV = SceneView->FOV;
+	float ViewRectWidth = SceneView->UnscaledViewRect.Width();
+	float ViewRectHeight = SceneView->UnscaledViewRect.Height();
+		
+	const FVector& ViewLocation = SceneView->ViewMatrices.GetViewOrigin();
+	
+	float NearPlaneSize = CalculateNearPlaneSize(FOV);
+	float DistAdjust = NearPlaneSize / RefNearPlaneWidth;
+	DistAdjust *= RefFrameBufferWidth / ViewRectWidth;
+
+	CameraInfos.Emplace(ViewLocation, SceneView->ViewMatrices.GetViewProjectionMatrix(), ViewRectWidth, DistAdjust);
+}
+
 float FImgMediaMipMapInfoManager::CalculateNearPlaneSize(float InFOV) const
 {
 	float NearPlaneDistance = RefNearPlaneDistance;
diff --git a/Engine/Plugins/Media/ImgMedia/Source/ImgMedia/Private/Assets/ImgMediaMipMapInfoManager.h b/Engine/Plugins/Media/ImgMedia/Source/ImgMedia/Private/Assets/ImgMediaMipMapInfoManager.h
--- a/Engine/Plugins/Media/ImgMedia/Source/ImgMedia/Private/Assets/ImgMediaMipMapInfoManager.h
+++ b/Engine/Plugins/Media/ImgMedia/Source/ImgMedia/Private/Assets/ImgMediaMipMapInfoManager.h
@@ -11,6 +11,7 @@
 #include "UObject/WeakObjectPtrTemplates.h"
 
 class AActor;
+class FSceneView;
 class IMediaPlayer;
 class UCameraComponent;
 class UImgMediaPlaybackComponent;
@@ -153,6 +154,13 @@ protected:
 	 */
 	void UpdateCameraInfo();
 
+	/**
+	 * Add a SceneView to our camera info.
+	 * 
+	 * @param SceneView		The scene view we want to add.
+	 */
+	void AddCameraInfo(FSceneView* SceneView);
+
 	/** Calculate the size of the near plane in physical space.
 	 *
 	 * @param InFOV FOV of camera in degrees.
diff --git a/Engine/Plugins/Media/ImgMedia/Source/ImgMedia/Private/Loader/ImgMediaLoader.cpp b/Engine/Plugins/Media/ImgMedia/Source/ImgMedia/Private/Loader/ImgMediaLoader.cpp
--- a/Engine/Plugins/Media/ImgMedia/Source/ImgMedia/Private/Loader/ImgMediaLoader.cpp
+++ b/Engine/Plugins/Media/ImgMedia/Source/ImgMedia/Private/Loader/ImgMediaLoader.cpp
@@ -808,7 +808,8 @@ void FImgMediaLoader::LoadSequence(const FString& SequencePath, const FFrameRate
 		}
 		else
 		{
-			MipMapInfo->SetTextureInfo(SequenceName, GetNumMipLevels(), SequenceDim);
+			MipMapInfo->SetTextureInfo(SequenceName, GetNumMipLevels(), NumTilesX, NumTilesY,
+				SequenceDim);
 			if (GetNumMipLevels() > 1)
 			{
 				UncompressedSize = (UncompressedSize * 4) / 3;
diff --git a/Engine/Plugins/Media/ImgMedia/Source/ImgMedia/Public/ImgMediaMipMapInfo.h b/Engine/Plugins/Media/ImgMedia/Source/ImgMedia/Public/ImgMediaMipMapInfo.h
--- a/Engine/Plugins/Media/ImgMedia/Source/ImgMedia/Public/ImgMediaMipMapInfo.h
+++ b/Engine/Plugins/Media/ImgMedia/Source/ImgMedia/Public/ImgMediaMipMapInfo.h
@@ -16,17 +16,21 @@ class UMediaTexture;
  */
 struct FImgMediaMipMapCameraInfo
 {
-	FImgMediaMipMapCameraInfo(FString InName, FVector InLocation, float InScreenSize, float InDistAdjust)
-		: Name(InName)
-		, Location(InLocation)
+	FImgMediaMipMapCameraInfo(const FVector& InLocation,
+		const FMatrix& InViewMatrix,
+		float InScreenSize, float InDistAdjust)
+		: Location(InLocation)
+		, ViewMatrix(InViewMatrix)
 		, ScreenSize(InScreenSize)
 		, DistAdjust(InDistAdjust)
 	{
 	}
-	/** Name of this camera. */
-	FString Name;
+
 	/** Position of camera. */
 	FVector Location;
+	/** View projection matrix of camera. */
+	FMatrix ViewMatrix;
+
 	/** Size of screen. */
 	float ScreenSize;
 	/** Adjustment needed to mip level distance calculations due to camera FOV, etc. */
@@ -158,7 +162,8 @@ public:
 	 * @param NumMipMaps Number of mipmaps in our image sequence.
 	 * @param Dim Dimensions of the textures in our image sequence.
 	 */
-	void SetTextureInfo(FName InSequenceName, int NumMipMaps, const FIntPoint& Dim);
+	void SetTextureInfo(FName InSequenceName, int32 NumMipMaps, int32 InNumTilesX, int32 InNumTilesY,
+		const FIntPoint& Dim);
 
 	/**
 	 * Get what mipmap level should be used.
@@ -230,6 +235,19 @@ protected:
 	 */
 	void UpdateMipLevelCache();
 
+	/** Decide which tiles are visible and update CachedTileSelection.
+	 *
+	 * @param CameraInfo				View info to compute visibility from.
+	 * @param ObjectLocation			Position of object.
+	 * @param NextTileXVector			Vector to go from one tile to the next in the X axis.
+	 * @param NextTileYVector			Vector to go from one tile to the next in the Y axis.
+	 * @param TileRadiusInWorldSpace	Radius of a bounding sphere for a tile.
+	 */
+	void CalculateTileVisibility(const FImgMediaMipMapCameraInfo& CameraInfo,
+		const FVector& ObjectLocation, const FVector& NextTileXVector,
+		const FVector& NextTileYVector, float TileRadiusInWorldSpace);
+
+
 	/**
 	 * Updates the MipLevelDistances based on current information.
 	 */
@@ -237,6 +255,10 @@ protected:
 
 	/** Name of this sequence. */
 	FName SequenceName;
+	/** Number of tiles in the X direction. */
+	int32 NumTilesX;
+	/** Number of tiles in the Y direction. */
+	int32 NumTilesY;
 
 	/** Ideal distance for mip level 0. */
 	float MipLevel0Distance;
