ImgMedia: Added process images tab.

#rnx

diff --git a/Engine/Plugins/Media/ImgMedia/Source/ImgMediaEditor/Private/ImgMediaEditorModule.cpp b/Engine/Plugins/Media/ImgMedia/Source/ImgMediaEditor/Private/ImgMediaEditorModule.cpp
--- a/Engine/Plugins/Media/ImgMedia/Source/ImgMediaEditor/Private/ImgMediaEditorModule.cpp
+++ b/Engine/Plugins/Media/ImgMedia/Source/ImgMediaEditor/Private/ImgMediaEditorModule.cpp
@@ -14,6 +14,7 @@
 #include "UObject/NameTypes.h"
 #include "Widgets/Docking/SDockTab.h"
 #include "Widgets/SImgMediaCache.h"
+#include "Widgets/SImgMediaProcessImages.h"
 #include "WorkspaceMenuStructure.h"
 #include "WorkspaceMenuStructureModule.h"
 
@@ -116,12 +117,21 @@ protected:
 			.SetDisplayName(LOCTEXT("ImgMediaCacheTabTitle", "Cache"))
 			.SetTooltipText(LOCTEXT("ImgMediaCacheTooltipText", "Open the cache tab."))
 			.SetIcon(FSlateIcon(FEditorStyle::GetStyleSetName(), "SequenceRecorder.TabIcon"));
+
+		// Add process images tab.
+		FGlobalTabmanager::Get()->RegisterNomadTabSpawner(ImgMediaProcessImagesTabName,
+			FOnSpawnTab::CreateStatic(&FImgMediaEditorModule::SpawnProcessImagesTab))
+			.SetGroup(MediaBrowserGroup)
+			.SetDisplayName(LOCTEXT("ImgMediaProcessImagesTabTitle", "Process Images"))
+			.SetTooltipText(LOCTEXT("ImgMediaProcessImagesTooltipText", "Open the Process Images tab."))
+			.SetIcon(FSlateIcon(FEditorStyle::GetStyleSetName(), "LevelEditor.Tabs.Viewports"));
 	}
 
 	void UnregisterTabSpawners()
 	{
 		if (FSlateApplication::IsInitialized())
 		{
+			FGlobalTabmanager::Get()->UnregisterNomadTabSpawner(ImgMediaProcessImagesTabName);
 			FGlobalTabmanager::Get()->UnregisterNomadTabSpawner(ImgMediaCacheTabName);
 		}
 	}
@@ -135,16 +145,29 @@ protected:
 			];
 	}
 
+	static TSharedRef<SDockTab> SpawnProcessImagesTab(const FSpawnTabArgs& SpawnTabArgs)
+	{
+		return SNew(SDockTab)
+			.TabRole(ETabRole::NomadTab)
+			[
+				SNew(SImgMediaProcessImages)
+			];
+	}
+
 private:
 
 	/** Customization name to avoid reusing staticstruct during shutdown. */
 	FName CustomizedStructName;
 	FName ImportInfoStructName;
 
+	/** Names for tabs. */
+	static FLazyName ImgMediaProcessImagesTabName;
+
 	/** The collection of registered asset type actions. */
 	TArray<TSharedRef<IAssetTypeActions>> RegisteredAssetTypeActions;
 };
 
+FLazyName FImgMediaEditorModule::ImgMediaProcessImagesTabName(TEXT("ImgMediaProcessImages"));
 
 IMPLEMENT_MODULE(FImgMediaEditorModule, ImgMediaEditor);
 
diff --git a/Engine/Plugins/Media/ImgMedia/Source/ImgMediaEditor/Private/Widgets/ImgMediaProcessImagesOptions.cpp b/Engine/Plugins/Media/ImgMedia/Source/ImgMediaEditor/Private/Widgets/ImgMediaProcessImagesOptions.cpp
new file mode 100644
--- /dev/null
+++ b/Engine/Plugins/Media/ImgMedia/Source/ImgMediaEditor/Private/Widgets/ImgMediaProcessImagesOptions.cpp
@@ -0,0 +1,9 @@
+// Copyright Epic Games, Inc. All Rights Reserved.
+
+#include "ImgMediaProcessImagesOptions.h"
+
+UImgMediaProcessImagesOptions::UImgMediaProcessImagesOptions(const FObjectInitializer& ObjectInitializer)
+	: Super(ObjectInitializer)
+{
+}
+
diff --git a/Engine/Plugins/Media/ImgMedia/Source/ImgMediaEditor/Private/Widgets/ImgMediaProcessImagesOptions.h b/Engine/Plugins/Media/ImgMedia/Source/ImgMediaEditor/Private/Widgets/ImgMediaProcessImagesOptions.h
new file mode 100644
--- /dev/null
+++ b/Engine/Plugins/Media/ImgMedia/Source/ImgMediaEditor/Private/Widgets/ImgMediaProcessImagesOptions.h
@@ -0,0 +1,37 @@
+// Copyright Epic Games, Inc. All Rights Reserved.
+
+/**
+ * Contains options for processing image sequences.
+ */
+
+#pragma once
+
+#include "CoreMinimal.h"
+#include "UObject/ObjectMacros.h"
+#include "ImgMediaProcessImagesOptions.generated.h"
+
+UCLASS(MinimalAPI)
+class UImgMediaProcessImagesOptions : public UObject
+{
+	GENERATED_UCLASS_BODY()
+public:
+	/** The directory that contains the image sequence files. */
+	UPROPERTY(EditAnywhere, Category = Sequence)
+	FFilePath SequencePath;
+	
+	/** Width of a tile in pixels. If 0, then do not make tiles. */
+	UPROPERTY(EditAnywhere, Transient, Category = Tiles)
+	int32 TileSizeX = 0;
+
+	/** Height of a tile in pixels. If 0, then do not make tiles. */
+	UPROPERTY(EditAnywhere, Transient, Category = Tiles)
+	int32 TileSizeY = 0;
+
+	/** Number of tiles in the X direction. If 0, then there are no tiles. */
+	UPROPERTY(VisibleAnywhere, Transient, Category = Tiles)
+	int32 NumTilesX = 0;
+
+	/** Number of tiles in the Y direction. If 0, then there are no tiles. */
+	UPROPERTY(VisibleAnywhere, Transient, Category = Tiles)
+	int32 NumTilesY = 0;
+};
diff --git a/Engine/Plugins/Media/ImgMedia/Source/ImgMediaEditor/Private/Widgets/SImgMediaProcessImages.cpp b/Engine/Plugins/Media/ImgMedia/Source/ImgMediaEditor/Private/Widgets/SImgMediaProcessImages.cpp
new file mode 100644
--- /dev/null
+++ b/Engine/Plugins/Media/ImgMedia/Source/ImgMediaEditor/Private/Widgets/SImgMediaProcessImages.cpp
@@ -0,0 +1,78 @@
+// Copyright Epic Games, Inc. All Rights Reserved.
+
+#include "Widgets/SImgMediaProcessImages.h"
+
+#include "EditorStyleSet.h"
+#include "ImgMediaProcessImagesOptions.h"
+#include "Modules/ModuleManager.h"
+#include "PropertyEditorModule.h"
+#include "SlateOptMacros.h"
+#include "UObject/ObjectMacros.h"
+#include "Widgets/Input/SButton.h"
+#include "Widgets/Layout/SBox.h"
+#include "Widgets/Layout/SScrollBox.h"
+#include "Widgets/Layout/SSplitter.h"
+#include "Widgets/Text/STextBlock.h"
+
+#define LOCTEXT_NAMESPACE "ImgMediaProcessImages"
+
+SImgMediaProcessImages::~SImgMediaProcessImages()
+{
+}
+
+BEGIN_SLATE_FUNCTION_BUILD_OPTIMIZATION
+void SImgMediaProcessImages::Construct(const FArguments& InArgs)
+{
+	// Set up widgets.
+	TSharedPtr<SBox> DetailsViewBox;
+	
+	ChildSlot
+	[
+		SNew(SVerticalBox)
+
+		+ SVerticalBox::Slot()
+			.Padding(0, 20, 0, 0)
+			.AutoHeight()
+
+		// Add details view.
+		+ SVerticalBox::Slot()
+			.AutoHeight()
+			[
+				SAssignNew(DetailsViewBox, SBox)
+			]
+			
+		// Add process images button.
+		+ SVerticalBox::Slot()
+			.AutoHeight()
+			.Padding(0, 5, 0, 0)
+			[
+				SNew(SButton)
+					.ButtonStyle(FEditorStyle::Get(), "HoverHintOnly")
+					.ForegroundColor(FSlateColor::UseForeground())
+					.OnClicked(this, &SImgMediaProcessImages::OnProcessImagesClicked)
+					.Text(LOCTEXT("StartProcessImages", "ProcessImages"))
+					.ToolTipText(LOCTEXT("ProcesssImagesButtonToolTip", "Start processing images."))
+			]
+	];
+
+	// Create object with our options.
+	Options = TStrongObjectPtr<UImgMediaProcessImagesOptions>(NewObject<UImgMediaProcessImagesOptions>(GetTransientPackage(), NAME_None));
+
+	// Create detail view with our options.
+	FPropertyEditorModule& PropertyEditorModule = FModuleManager::GetModuleChecked<FPropertyEditorModule>("PropertyEditor");
+	FDetailsViewArgs DetailsViewArgs;
+	DetailsViewArgs.bAllowSearch = false;
+	DetailsViewArgs.NameAreaSettings = FDetailsViewArgs::HideNameArea;
+	DetailsView = PropertyEditorModule.CreateDetailView(DetailsViewArgs);
+	DetailsView->SetObject(Options.Get());
+
+	DetailsViewBox->SetContent(DetailsView->AsShared());
+}
+END_SLATE_FUNCTION_BUILD_OPTIMIZATION
+
+FReply SImgMediaProcessImages::OnProcessImagesClicked()
+{
+	return FReply::Handled();
+}
+
+#undef LOCTEXT_NAMESPACE
diff --git a/Engine/Plugins/Media/ImgMedia/Source/ImgMediaEditor/Private/Widgets/SImgMediaProcessImages.h b/Engine/Plugins/Media/ImgMedia/Source/ImgMediaEditor/Private/Widgets/SImgMediaProcessImages.h
new file mode 100644
--- /dev/null
+++ b/Engine/Plugins/Media/ImgMedia/Source/ImgMediaEditor/Private/Widgets/SImgMediaProcessImages.h
@@ -0,0 +1,34 @@
+// Copyright Epic Games, Inc. All Rights Reserved.
+
+#pragma once
+
+#include "CoreMinimal.h"
+#include "ImgMediaProcessImagesOptions.h"
+#include "Input/Reply.h"
+#include "UObject/StrongObjectPtr.h"
+#include "Widgets/SCompoundWidget.h"
+
+class IDetailsView;
+
+/**
+ * SImgMediaProcessImages provides processing of image sequences.
+ */
+class SImgMediaProcessImages : public SCompoundWidget
+{
+public:
+	SLATE_BEGIN_ARGS(SImgMediaProcessImages){}
+	SLATE_END_ARGS()
+
+	virtual ~SImgMediaProcessImages();
+
+	void Construct(const FArguments& InArgs);
+
+private:
+	/** Called when we click on the process images button. */
+	FReply OnProcessImagesClicked();
+
+	/** Holds our details view. */
+	TSharedPtr<class IDetailsView> DetailsView;
+	/** Object that holds our options. */
+	TStrongObjectPtr<UImgMediaProcessImagesOptions> Options;
+};
