ImgMedia: Process images button uses tiling code from ImgMediaSource to create tiles.
#rnx

diff --git a/Engine/Plugins/Media/ImgMedia/Source/ImgMediaEditor/Private/Widgets/ImgMediaProcessImagesOptions.h b/Engine/Plugins/Media/ImgMedia/Source/ImgMediaEditor/Private/Widgets/ImgMediaProcessImagesOptions.h
--- a/Engine/Plugins/Media/ImgMedia/Source/ImgMediaEditor/Private/Widgets/ImgMediaProcessImagesOptions.h
+++ b/Engine/Plugins/Media/ImgMedia/Source/ImgMediaEditor/Private/Widgets/ImgMediaProcessImagesOptions.h
@@ -18,6 +18,10 @@ public:
 	/** The directory that contains the image sequence files. */
 	UPROPERTY(EditAnywhere, Category = Sequence)
 	FFilePath SequencePath;
+
+	/** The directory to output the processed image sequence files to. */
+	UPROPERTY(EditAnywhere, Category = Sequence)
+	FDirectoryPath OutputPath;
 	
 	/** Width of a tile in pixels. If 0, then do not make tiles. */
 	UPROPERTY(EditAnywhere, Transient, Category = Tiles)
diff --git a/Engine/Plugins/Media/ImgMedia/Source/ImgMediaEditor/Private/Widgets/SImgMediaProcessImages.cpp b/Engine/Plugins/Media/ImgMedia/Source/ImgMediaEditor/Private/Widgets/SImgMediaProcessImages.cpp
--- a/Engine/Plugins/Media/ImgMedia/Source/ImgMediaEditor/Private/Widgets/SImgMediaProcessImages.cpp
+++ b/Engine/Plugins/Media/ImgMedia/Source/ImgMediaEditor/Private/Widgets/SImgMediaProcessImages.cpp
@@ -2,8 +2,16 @@
 
 #include "Widgets/SImgMediaProcessImages.h"
 
+#include "Async/Async.h"
 #include "EditorStyleSet.h"
+#include "Framework/Notifications/NotificationManager.h"
+#include "HAL/PlatformFileManager.h"
+#include "IImageWrapperModule.h"
+#include "ImageWrapperHelper.h"
+#include "ImgMediaEditorModule.h"
 #include "ImgMediaProcessImagesOptions.h"
+#include "Misc/FileHelper.h"
+#include "Misc/Paths.h"
 #include "Modules/ModuleManager.h"
 #include "PropertyEditorModule.h"
 #include "SlateOptMacros.h"
@@ -12,6 +20,7 @@
 #include "Widgets/Layout/SBox.h"
 #include "Widgets/Layout/SScrollBox.h"
 #include "Widgets/Layout/SSplitter.h"
+#include "Widgets/Notifications/SNotificationList.h"
 #include "Widgets/Text/STextBlock.h"
 
 #define LOCTEXT_NAMESPACE "ImgMediaProcessImages"
@@ -44,13 +53,12 @@ void SImgMediaProcessImages::Construct(const FArguments& InArgs)
 		// Add process images button.
 		+ SVerticalBox::Slot()
 			.AutoHeight()
-			.Padding(0, 5, 0, 0)
+			.Padding(4.0f)
+			.HAlign(HAlign_Left)
 			[
 				SNew(SButton)
-					.ButtonStyle(FEditorStyle::Get(), "HoverHintOnly")
-					.ForegroundColor(FSlateColor::UseForeground())
 					.OnClicked(this, &SImgMediaProcessImages::OnProcessImagesClicked)
-					.Text(LOCTEXT("StartProcessImages", "ProcessImages"))
+					.Text(LOCTEXT("StartProcessImages", "Process Images"))
 					.ToolTipText(LOCTEXT("ProcesssImagesButtonToolTip", "Start processing images."))
 			]
 	];
@@ -72,7 +80,169 @@ END_SLATE_FUNCTION_BUILD_OPTIMIZATION
 
 FReply SImgMediaProcessImages::OnProcessImagesClicked()
 {
+	// Create notification.
+	FNotificationInfo Info(FText::GetEmpty());
+	Info.bFireAndForget = false;
+	TSharedPtr<SNotificationItem> ConfirmNotification = FSlateNotificationManager::Get().AddNotification(Info);
+
+	// Start async task to process files.
+	Async(EAsyncExecution::Thread, [this, ConfirmNotification]()
+	{
+		ProcessAllImages(ConfirmNotification);
+	});
+
 	return FReply::Handled();
 }
 
+void SImgMediaProcessImages::ProcessAllImages(TSharedPtr<SNotificationItem> ConfirmNotification)
+{
+	int32 InTileWidth = Options->TileSizeX;
+	int32 InTileHeight = Options->TileSizeY;
+
+	// Create output directory.
+	FString OutPath = Options->OutputPath.Path;
+	IPlatformFile& PlatformFile = FPlatformFileManager::Get().GetPlatformFile();
+	PlatformFile.CreateDirectoryTree(*OutPath);
+
+	// Get source files.
+	FString SequencePath = FPaths::GetPath(Options->SequencePath.FilePath);
+	
+	TArray<FString> FoundFiles;
+	IFileManager::Get().FindFiles(FoundFiles, *SequencePath, TEXT("*"));
+	FoundFiles.Sort();
+	UE_LOG(LogImgMediaEditor, Warning, TEXT("Found %i image files in %s to import."), FoundFiles.Num(), *SequencePath);
+	if (FoundFiles.Num() == 0)
+	{
+		UE_LOG(LogImgMediaEditor, Error, TEXT("No files to import."));
+	}
+	else
+	{
+		// Create image wrapper
+		FString Ext = FPaths::GetExtension(FoundFiles[0]);
+		EImageFormat ImageFormat = ImageWrapperHelper::GetImageFormat(Ext);
+
+		if (ImageFormat == EImageFormat::Invalid)
+		{
+			UE_LOG(LogImgMediaEditor, Error, TEXT("Invalid file format %s"), *Ext);
+		}
+		else
+		{
+			IImageWrapperModule& ImageWrapperModule = FModuleManager::LoadModuleChecked<IImageWrapperModule>("ImageWrapper");
+			TSharedPtr<IImageWrapper> ImageWrapper = ImageWrapperModule.CreateImageWrapper(ImageFormat);
+
+			// Loop through all files.
+			int NumDone = 0;
+			int TotalNum = FoundFiles.Num();
+			for (const FString& FileName : FoundFiles)
+			{
+				// Update notification with current status.
+				Async(EAsyncExecution::TaskGraphMainThread, [ConfirmNotification, NumDone, TotalNum]()
+				{
+					if (ConfirmNotification.IsValid())
+					{
+						ConfirmNotification->SetText(
+							FText::Format(LOCTEXT("ImgMediaCompleted", "ImgMedia Completed {0}/{1}"),
+								FText::AsNumber(NumDone), FText::AsNumber(TotalNum)));
+					}
+				});
+				NumDone++;
+
+				FString FullFileName = FPaths::Combine(SequencePath, FileName);
+
+				// Load image into buffer.
+				TArray64<uint8> InputBuffer;
+				if (!FFileHelper::LoadFileToArray(InputBuffer, *FullFileName))
+				{
+					UE_LOG(LogImgMediaEditor, Error, TEXT("Failed to load %s"), *FullFileName);
+					break;
+				}
+				if (!ImageWrapper.IsValid() || !ImageWrapper->SetCompressed(InputBuffer.GetData(), InputBuffer.Num()))
+				{
+					UE_LOG(LogImgMediaEditor, Error, TEXT("Failed to create image wrapper for %s"), *FullFileName);
+					break;
+				}
+
+				// Import this image.
+				FString Name = FPaths::ChangeExtension(FPaths::Combine(OutPath, FileName), TEXT(""));
+				ProcessImage(ImageWrapper, InTileWidth, InTileHeight, Name, Ext);
+			}
+		}
+	}
+
+	// Close notification. Must be run on the main thread.
+	Async(EAsyncExecution::TaskGraphMainThread, [ConfirmNotification]()
+	{
+		if (ConfirmNotification.IsValid())
+		{
+			ConfirmNotification->SetEnabled(false);
+			ConfirmNotification->SetCompletionState(SNotificationItem::CS_Success);
+			ConfirmNotification->ExpireAndFadeout();
+		}
+	});
+}
+
+void SImgMediaProcessImages::ProcessImage(TSharedPtr<IImageWrapper>& InImageWrapper,
+	int32 InTileWidth, int32 InTileHeight, const FString& InName, const FString& FileExtension)
+{
+	// Get image data.
+	ERGBFormat Format = InImageWrapper->GetFormat();
+	int32 Width = InImageWrapper->GetWidth();
+	int32 Height = InImageWrapper->GetHeight();
+	int32 BitDepth = InImageWrapper->GetBitDepth();
+	TArray64<uint8> RawData;
+	InImageWrapper->GetRaw(Format, BitDepth, RawData);
+
+	int32 NumTilesX = InTileWidth > 0 ? Width / InTileWidth : 1;
+	int32 NumTilesY = InTileHeight > 0 ? Height / InTileHeight : 1;
+	int32 TileWidth = Width / NumTilesX;
+	int32 TileHeight = Height / NumTilesY;
+	int32 BytesPerPixel = RawData.Num() / (Width * Height);
+	TArray64<uint8> TileRawData;
+	TileRawData.AddZeroed(TileWidth * TileHeight * BytesPerPixel);
+	bool bIsTiled = (NumTilesX > 1) || (NumTilesY > 1);
+
+	// Create a directory if we have tiles.
+	FString FileName;
+	if (bIsTiled)
+	{
+		IPlatformFile& PlatformFile = FPlatformFileManager::Get().GetPlatformFile();
+		PlatformFile.CreateDirectoryTree(*InName);
+		FileName = FPaths::Combine(InName, FPaths::GetCleanFilename(InName));;
+	}
+	else
+	{
+		FileName = InName;
+	}
+
+	// Loop over y tiles.
+	for (int TileY = 0; TileY < NumTilesY; ++TileY)
+	{
+		// Loop over x tiles.
+		for (int TileX = 0; TileX < NumTilesX; ++TileX)
+		{
+			// Copy tile line by line.
+			uint8* DestPtr = TileRawData.GetData();
+			uint8* SrcPtr = RawData.GetData() + TileX * TileWidth * BytesPerPixel +
+				TileY * TileHeight * Width * BytesPerPixel;
+			for (int LineY = 0; LineY < TileHeight; ++LineY)
+			{
+				FMemory::Memcpy(DestPtr, SrcPtr,
+					TileWidth * BytesPerPixel);
+				DestPtr += TileWidth * BytesPerPixel;
+				SrcPtr += Width * BytesPerPixel;
+			}
+
+			// Compress data.
+			InImageWrapper->SetRaw(TileRawData.GetData(), TileRawData.Num(),
+				TileWidth, TileHeight, Format, BitDepth);
+			const TArray64<uint8> CompressedData = InImageWrapper->GetCompressed((int32)EImageCompressionQuality::Uncompressed);
+
+			// Write out tile.
+			FString Name = FString::Format(TEXT("{0}_x{1}_y{2}.{3}"),
+				{*FileName, TileX, TileY, *FileExtension});
+			FFileHelper::SaveArrayToFile(CompressedData, *Name);
+		}
+	}
+}
+
 #undef LOCTEXT_NAMESPACE
diff --git a/Engine/Plugins/Media/ImgMedia/Source/ImgMediaEditor/Private/Widgets/SImgMediaProcessImages.h b/Engine/Plugins/Media/ImgMedia/Source/ImgMediaEditor/Private/Widgets/SImgMediaProcessImages.h
--- a/Engine/Plugins/Media/ImgMedia/Source/ImgMediaEditor/Private/Widgets/SImgMediaProcessImages.h
+++ b/Engine/Plugins/Media/ImgMedia/Source/ImgMediaEditor/Private/Widgets/SImgMediaProcessImages.h
@@ -9,6 +9,8 @@
 #include "Widgets/SCompoundWidget.h"
 
 class IDetailsView;
+class IImageWrapper;
+class SNotificationItem;
 
 /**
  * SImgMediaProcessImages provides processing of image sequences.
@@ -27,6 +29,27 @@ private:
 	/** Called when we click on the process images button. */
 	FReply OnProcessImagesClicked();
 
+	/**
+	 * Processes all images in the sequence and generate tiles/mips.
+	 *
+	 * @param ConfirmNotification	Notification that will be updated with progress and closed when done.
+	 */
+	void ProcessAllImages(TSharedPtr<SNotificationItem> ConfirmNotification);
+	
+	/**
+	 * Processess a single image and writes out 1 or more files.
+	 * Tiles and mips may be generated.
+	 * This does NOT run on the game thread.
+	 *
+	 * @param InImageWrapper	ImageWrapper to read/write the image.
+	 * @param InTileWidth		Desired width of tiles.
+	 * @param InTileHeight		Desired height of tiles.
+	 * @param InName			Full path and name of file to write to WITHOUT the extension (e.g. no .exr)
+	 * @param FileExtension		Extension to append to the file name.
+	 */
+	void ProcessImage(TSharedPtr<IImageWrapper>& InImageWrapper,
+		int32 InTileWidth, int32 InTileHeight, const FString& InName, const FString& FileExtension);
+
 	/** Holds our details view. */
 	TSharedPtr<class IDetailsView> DetailsView;
 	/** Object that holds our options. */
