ImgMedia: Loader does not need the media source to get the number of tiles.
#rnx

diff --git a/Engine/Plugins/Media/ImgMedia/Source/ImgMedia/Private/Loader/ImgMediaLoader.cpp b/Engine/Plugins/Media/ImgMedia/Source/ImgMedia/Private/Loader/ImgMediaLoader.cpp
--- a/Engine/Plugins/Media/ImgMedia/Source/ImgMedia/Private/Loader/ImgMediaLoader.cpp
+++ b/Engine/Plugins/Media/ImgMedia/Source/ImgMedia/Private/Loader/ImgMediaLoader.cpp
@@ -63,6 +63,7 @@ FImgMediaLoader::FImgMediaLoader(const TSharedRef<FImgMediaScheduler, ESPMode::T
 	, ImageWrapperModule(FModuleManager::LoadModuleChecked<IImageWrapperModule>("ImageWrapper"))
 	, Initialized(false)
 	, bFillGapsInSequence(bInFillGapsInSequence)
+	, bIsTiled(false)
 	, NumTilesX(0)
 	, NumTilesY(0)
 	, NumLoadAhead(0)
@@ -623,7 +624,7 @@ IQueuedWork* FImgMediaLoader::GetWork()
 
 
 void FImgMediaLoader::Initialize(const FString& SequencePath, const FFrameRate& FrameRateOverride,
-	bool Loop, int32 InNumTilesX, int32 InNumTilesY)
+	bool Loop)
 {
 	UE_LOG(LogImgMedia, Verbose, TEXT("Loader %p: Initializing with %s (FrameRateOverride = %s, Loop = %i)"),
 		this,
@@ -634,15 +635,6 @@ void FImgMediaLoader::Initialize(const FString& SequencePath, const FFrameRate&
 
 	check(!Initialized); // reinitialization not allowed for now
 
-	// Set up tile info.
-	NumTilesX = InNumTilesX;
-	NumTilesY = InNumTilesY;
-	if ((NumTilesX < 1) || (NumTilesY < 1))
-	{
-		NumTilesX = 1;
-		NumTilesY = 1;
-	}
-
 	LoadSequence(SequencePath, FrameRateOverride, Loop);
 	FPlatformMisc::MemoryBarrier();
 
@@ -723,6 +715,9 @@ void FImgMediaLoader::LoadSequence(const FString& SequencePath, const FFrameRate
 	// Get mips.
 	FindMips(SequencePath);
 
+	// Get tile info.
+	GetTileInformation();
+
 	FScopeLock Lock(&CriticalSection);
 
 	// create image reader
@@ -855,10 +850,19 @@ void FImgMediaLoader::FindFiles(const FString& SequencePath, TArray<FString>& Ou
 	IFileManager::Get().FindFiles(FoundFiles, *SequencePath, TEXT("*"));
 
 	// If we did not find any files, then maybe we have tile directories.
-	bool bIsTiled = FoundFiles.Num() == 0;
-	if (bIsTiled)
+	// If we already have tiles then don't do anything.
+	if (bIsTiled == false)
 	{
-		IFileManager::Get().FindFiles(FoundFiles, *(SequencePath / TEXT("*")), false, true);
+		bIsTiled = FoundFiles.Num() == 0;
+		if (bIsTiled)
+		{
+			IFileManager::Get().FindFiles(FoundFiles, *(SequencePath / TEXT("*")), false, true);
+			if (FoundFiles.Num() == 0)
+			{
+				// We really got nothing.
+				bIsTiled = false;
+			}
+		}
 	}
 
 	UE_LOG(LogImgMedia, Verbose, TEXT("Loader %p: Found %i image files in %s"), this, FoundFiles.Num(), *SequencePath);
@@ -977,6 +981,49 @@ void FImgMediaLoader::FindMips(const FString& SequencePath)
 	}
 }
 
+void FImgMediaLoader::GetTileInformation()
+{
+	NumTilesX = 1;
+	NumTilesY = 1;
+
+	// Do we have tiles?
+	if (bIsTiled)
+	{
+		// Get the last file as that will tell us the number of tiles.
+		FString Path = FPaths::GetPath(ImagePaths[0][0]);
+		TArray<FString> FoundFiles;
+		IFileManager::Get().FindFiles(FoundFiles, *Path, TEXT("*"));
+		FoundFiles.Sort();
+		if (FoundFiles.Num() > 0)
+		{
+			// Remove the file extension.
+			Path = FPaths::ChangeExtension(FoundFiles[FoundFiles.Num() - 1], TEXT(""));
+
+			// Filename should be ????_x0_y0.
+			// Find the last y.
+			int32 YIndex = 0;
+			if (Path.FindLastChar(TEXT('y'), YIndex))
+			{
+				// Extract out the number.
+				int32 YLength = Path.Len() - (YIndex + 1);
+				FString NumberString = Path.Right(YLength);
+				// We count from 0, so add 1 to get the number of tiles.
+				NumTilesY = FCString::Atoi(*NumberString) + 1;
+
+				// Find the last x.
+				int32 XIndex = 0;
+				if (Path.FindLastChar(TEXT('x'), XIndex))
+				{
+					// Extract out the number.
+					XIndex++;
+					int32 XLength = Path.Len() - XIndex - YLength - 2;
+					NumberString = Path.Mid(XIndex, XLength);
+					NumTilesX = FCString::Atoi(*NumberString) + 1;
+				}
+			}
+		}
+	}
+}
 
 uint32 FImgMediaLoader::TimeToFrameNumber(FTimespan Time) const
 {
diff --git a/Engine/Plugins/Media/ImgMedia/Source/ImgMedia/Private/Loader/ImgMediaLoader.h b/Engine/Plugins/Media/ImgMedia/Source/ImgMedia/Private/Loader/ImgMediaLoader.h
--- a/Engine/Plugins/Media/ImgMedia/Source/ImgMedia/Private/Loader/ImgMediaLoader.h
+++ b/Engine/Plugins/Media/ImgMedia/Source/ImgMedia/Private/Loader/ImgMediaLoader.h
@@ -239,12 +239,9 @@ public:
 	 * @param SequencePath Path to the image sequence.
 	 * @param FrameRateOverride The frame rate to use (0/0 = do not override).
 	 * @param Loop Whether the cache should loop around.
-	 * @param InNumTilesX Number of tiles in the X direction.
-	 * @param InNumTilesY Number of tiles in the Y direction.
 	 * @see IsInitialized
 	 */
-	void Initialize(const FString& SequencePath, const FFrameRate& FrameRateOverride, bool Loop,
-		int32 InNumTilesX, int32 InNumTilesY);
+	void Initialize(const FString& SequencePath, const FFrameRate& FrameRateOverride, bool Loop);
 
 	/**
 	 * Whether this loader has been initialized yet.
@@ -339,6 +336,11 @@ protected:
 	 */
 	void FindMips(const FString& SequencePath);
 
+	/**
+	 * Gets information for tiling from the sequence.
+	 */
+	void GetTileInformation();
+
 	/**
 	 * Get the frame number corresponding to the specified play head time.
 	 *
@@ -446,6 +448,9 @@ private:
 	/** If true, then any gaps in the sequence will be filled with blank frames. */
 	bool bFillGapsInSequence;
 
+	/** True if this sequence contains tiles. */
+	bool bIsTiled;
+
 	/** Number of tiles in the X direction. */
 	int32 NumTilesX;
 
diff --git a/Engine/Plugins/Media/ImgMedia/Source/ImgMedia/Private/Player/ImgMediaPlayer.cpp b/Engine/Plugins/Media/ImgMedia/Source/ImgMedia/Private/Player/ImgMediaPlayer.cpp
--- a/Engine/Plugins/Media/ImgMedia/Source/ImgMedia/Private/Player/ImgMediaPlayer.cpp
+++ b/Engine/Plugins/Media/ImgMedia/Source/ImgMedia/Private/Player/ImgMediaPlayer.cpp
@@ -182,15 +182,11 @@ bool FImgMediaPlayer::Open(const FString& Url, const IMediaOptions* Options)
 	FFrameRate FrameRateOverride(0, 0);
 	TSharedPtr<FImgMediaMipMapInfo, ESPMode::ThreadSafe> MipMapInfo;
 	bool bFillGapsInSequence = true;
-	int32 NumTilesX = 0;
-	int32 NumTilesY = 0;
 	if (Options != nullptr)
 	{
 		FrameRateOverride.Denominator = Options->GetMediaOption(ImgMedia::FrameRateOverrideDenonimatorOption, 0LL);
 		FrameRateOverride.Numerator = Options->GetMediaOption(ImgMedia::FrameRateOverrideNumeratorOption, 0LL);
 		bFillGapsInSequence = Options->GetMediaOption(ImgMedia::FillGapsInSequenceOption, true);
-		NumTilesX = Options->GetMediaOption(ImgMedia::NumTilesXOption, (int64)0);
-		NumTilesY = Options->GetMediaOption(ImgMedia::NumTilesYOption, (int64)0);
 
 		TSharedPtr<IMediaOptions::FDataContainer, ESPMode::ThreadSafe> DefaultValue;
 		TSharedPtr<IMediaOptions::FDataContainer, ESPMode::ThreadSafe> DataContainer = Options->GetMediaOption(ImgMedia::MipMapInfoOption, DefaultValue);
@@ -218,8 +214,7 @@ bool FImgMediaPlayer::Open(const FString& Url, const IMediaOptions* Options)
 
 	const FString SequencePath = Url.RightChop(6);
 
-	Async(EAsyncExecution::ThreadPool, [FrameRateOverride, LoaderPtr = TWeakPtr<FImgMediaLoader,ESPMode::ThreadSafe>(Loader),
-		Proxy, SequencePath, Loop = ShouldLoop, NumTilesX, NumTilesY]()
+	Async(EAsyncExecution::ThreadPool, [FrameRateOverride, LoaderPtr = TWeakPtr<FImgMediaLoader, ESPMode::ThreadSafe>(Loader), Proxy, SequencePath, Loop = ShouldLoop]()
 	{
 		TSharedPtr<FImgMediaLoader, ESPMode::ThreadSafe> PinnedLoader = LoaderPtr.Pin();
 
@@ -232,7 +227,7 @@ bool FImgMediaPlayer::Open(const FString& Url, const IMediaOptions* Options)
 				ProxyPath = SequencePath; // fall back to root folder
 			}
 
-			PinnedLoader->Initialize(ProxyPath, FrameRateOverride, Loop, NumTilesX, NumTilesY);
+			PinnedLoader->Initialize(ProxyPath, FrameRateOverride, Loop);
 		}
 	});
 
