ImageMedia: Fixed a caching issue for reading tiled exrs that are stored as individual files.

diff --git a/Engine/Plugins/Media/ImgMedia/Source/ImgMedia/Private/Readers/ExrImgMediaReader.cpp b/Engine/Plugins/Media/ImgMedia/Source/ImgMedia/Private/Readers/ExrImgMediaReader.cpp
--- a/Engine/Plugins/Media/ImgMedia/Source/ImgMedia/Private/Readers/ExrImgMediaReader.cpp
+++ b/Engine/Plugins/Media/ImgMedia/Source/ImgMedia/Private/Readers/ExrImgMediaReader.cpp
@@ -47,11 +47,25 @@ FExrImgMediaReader::~FExrImgMediaReader()
 
 bool FExrImgMediaReader::GetFrameInfo(const FString& ImagePath, FImgMediaFrameInfo& OutInfo)
 {
+	TSharedPtr<FImgMediaLoader, ESPMode::ThreadSafe> Loader = LoaderPtr.Pin();
+	if (Loader.IsValid() == false)
+	{
+		return false;
+	}
 
+	// Get tile info.
+	int32 NumTilesX = Loader->GetNumTilesX();
+	int32 NumTilesY = Loader->GetNumTilesY();
 	FRgbaInputFile InputFile(ImagePath, 2);
 
-
-	return GetInfo(InputFile, OutInfo);
+	if (GetInfo(InputFile, OutInfo))
+	{
+		OutInfo.Dim.X *= NumTilesX;
+		OutInfo.Dim.Y *= NumTilesY;
+		OutInfo.UncompressedSize *= NumTilesY * NumTilesX;
+		return true;
+	}
+	return false;
 }
 
 
@@ -87,9 +101,10 @@ bool FExrImgMediaReader::ReadFrame(int32 FrameId, int32 MipLevel, const FImgMedi
 			return false;
 		}
 
-		// If we have tiles, then this is the size of just a tile, so multiply to get the full size.
+		// If we have tiles, then this is the size of just a tile, so multiply to get the full size. Same goes for size.
 		OutFrame->Info.Dim.X *= NumTilesX;
 		OutFrame->Info.Dim.Y *= NumTilesY;
+		OutFrame->Info.UncompressedSize *= NumTilesX * NumTilesY;
 
 		const FIntPoint& Dim = OutFrame->Info.Dim;
 
