ImgMedia: Removed tile and import info from media source.
#rnx

diff --git a/Engine/Plugins/Media/ImgMedia/Source/ImgMedia/Private/Assets/ImgMediaSource.cpp b/Engine/Plugins/Media/ImgMedia/Source/ImgMedia/Private/Assets/ImgMediaSource.cpp
--- a/Engine/Plugins/Media/ImgMedia/Source/ImgMedia/Private/Assets/ImgMediaSource.cpp
+++ b/Engine/Plugins/Media/ImgMedia/Source/ImgMedia/Private/Assets/ImgMediaSource.cpp
@@ -113,16 +113,6 @@ int64 UImgMediaSource::GetMediaOption(const FName& Key, int64 DefaultValue) cons
 		return FrameRateOverride.Numerator;
 	}
 
-	if (Key == ImgMedia::NumTilesXOption)
-	{
-		return ImportInfo.NumTilesX;
-	}
-
-	if (Key == ImgMedia::NumTilesYOption)
-	{
-		return ImportInfo.NumTilesY;
-	}
-
 	return Super::GetMediaOption(Key, DefaultValue);
 }
 
@@ -153,9 +143,7 @@ bool UImgMediaSource::HasMediaOption(const FName& Key) const
 		(Key == ImgMedia::FrameRateOverrideDenonimatorOption) ||
 		(Key == ImgMedia::FrameRateOverrideNumeratorOption) ||
 		(Key == ImgMedia::ProxyOverrideOption) ||
-		(Key == ImgMedia::MipMapInfoOption) ||
-		(Key == ImgMedia::NumTilesXOption) ||
-		(Key == ImgMedia::NumTilesYOption))
+		(Key == ImgMedia::MipMapInfoOption))
 	{
 		return true;
 	}
@@ -201,12 +189,6 @@ void UImgMediaSource::PostEditChangeProperty(FPropertyChangedEvent& PropertyChan
 
 FString UImgMediaSource::GetFullPath() const
 {
-	// Use the import path?
-	if ((ImportInfo.bIsUsable) && (ImportInfo.DestinationPath.Path.IsEmpty() == false))
-	{
-		return ImportInfo.DestinationPath.Path;
-	}
-
 	if (!FPaths::IsRelative(SequencePath.Path))
 	{
 		return SequencePath.Path;
diff --git a/Engine/Plugins/Media/ImgMedia/Source/ImgMedia/Public/ImgMediaSource.h b/Engine/Plugins/Media/ImgMedia/Source/ImgMedia/Public/ImgMediaSource.h
--- a/Engine/Plugins/Media/ImgMedia/Source/ImgMedia/Public/ImgMediaSource.h
+++ b/Engine/Plugins/Media/ImgMedia/Source/ImgMedia/Public/ImgMediaSource.h
@@ -22,33 +22,6 @@ struct IMGMEDIA_API FImgMediaSourceCustomizationSequenceProxy
 	GENERATED_BODY()
 };
 
-/** Information on the images that have been imported. */
-USTRUCT(BlueprintType)
-struct IMGMEDIA_API FImgMediaSourceImportInfo
-{
-	GENERATED_BODY()
-
-	/** Where to place the imported images. */
-	UPROPERTY(BlueprintReadOnly, EditAnywhere, Category = Sequence)
-	FDirectoryPath DestinationPath;
-
-	/** True if we are overriding destination path and not using the default. */
-	UPROPERTY(BlueprintReadOnly, EditAnywhere, Category = Sequence)
-	bool bIsDestinationPathOverriden = false;
-
-	/** True if the imported imagess are to be used. */
-	UPROPERTY(BlueprintReadOnly, EditAnywhere, Category = Sequence)
-	bool bIsUsable = false;
-
-	/** Number of tiles in the X direction. If 0, then there are no tiles. */
-	UPROPERTY(BlueprintReadOnly, EditAnywhere, Category = Sequence)
-	int32 NumTilesX = 0;
-
-	/** Number if tiles in the Y direction. If 0, then there are no tiles. */
-	UPROPERTY(BlueprintReadOnly, EditAnywhere, Category = Sequence)
-	int32 NumTilesY = 0;
-};
-
 /**
  * Media source for EXR image sequences.
  *
@@ -101,10 +74,6 @@ public:
 
 #endif // WITH_EDITORONLY_DATA
 
-	/** Info concerning the imported images. */
-	UPROPERTY(BlueprintReadWrite, EditAnywhere, Category = Sequence)
-	FImgMediaSourceImportInfo ImportInfo;
-
 public:
 
 	/**
diff --git a/Engine/Plugins/Media/ImgMedia/Source/ImgMediaEditor/Private/Customizations/ImgMediaSourceCustomizationImportInfo.cpp b/Engine/Plugins/Media/ImgMedia/Source/ImgMediaEditor/Private/Customizations/ImgMediaSourceCustomizationImportInfo.cpp
deleted file mode 100644
diff --git a/Engine/Plugins/Media/ImgMedia/Source/ImgMediaEditor/Private/Customizations/ImgMediaSourceCustomizationImportInfo.h b/Engine/Plugins/Media/ImgMedia/Source/ImgMediaEditor/Private/Customizations/ImgMediaSourceCustomizationImportInfo.h
deleted file mode 100644
diff --git a/Engine/Plugins/Media/ImgMedia/Source/ImgMediaEditor/Private/ImgMediaEditorModule.cpp b/Engine/Plugins/Media/ImgMedia/Source/ImgMediaEditor/Private/ImgMediaEditorModule.cpp
--- a/Engine/Plugins/Media/ImgMedia/Source/ImgMediaEditor/Private/ImgMediaEditorModule.cpp
+++ b/Engine/Plugins/Media/ImgMedia/Source/ImgMediaEditor/Private/ImgMediaEditorModule.cpp
@@ -7,7 +7,6 @@
 #include "AssetToolsModule.h"
 #include "AssetTools/ImgMediaSourceActions.h"
 #include "Customizations/ImgMediaSourceCustomization.h"
-#include "Customizations/ImgMediaSourceCustomizationImportInfo.h"
 #include "IAssetTools.h"
 #include "ImgMediaSource.h"
 #include "PropertyEditorModule.h"
@@ -54,13 +53,11 @@ protected:
 	void RegisterCustomizations()
 	{
 		CustomizedStructName = FImgMediaSourceCustomizationSequenceProxy::StaticStruct()->GetFName();
-		ImportInfoStructName = FImgMediaSourceImportInfo::StaticStruct()->GetFName();
 
 		FPropertyEditorModule& PropertyModule = FModuleManager::LoadModuleChecked<FPropertyEditorModule>("PropertyEditor");
 		{
 #if WITH_EDITORONLY_DATA
 			PropertyModule.RegisterCustomPropertyTypeLayout(CustomizedStructName, FOnGetPropertyTypeCustomizationInstance::CreateStatic(&FImgMediaSourceCustomization::MakeInstance));
-			PropertyModule.RegisterCustomPropertyTypeLayout(ImportInfoStructName, FOnGetPropertyTypeCustomizationInstance::CreateStatic(&FImgMediaSourceCustomizationImportInfo::MakeInstance));
 #endif // WITH_EDITORONLY_DATA
 		}
 	}
@@ -71,7 +68,6 @@ protected:
 		FPropertyEditorModule& PropertyModule = FModuleManager::LoadModuleChecked<FPropertyEditorModule>("PropertyEditor");
 		{
 #if WITH_EDITORONLY_DATA
-			PropertyModule.UnregisterCustomPropertyTypeLayout(ImportInfoStructName);
 			PropertyModule.UnregisterCustomPropertyTypeLayout(CustomizedStructName);
 #endif // WITH_EDITORONLY_DATA
 		}
@@ -158,7 +154,6 @@ private:
 
 	/** Customization name to avoid reusing staticstruct during shutdown. */
 	FName CustomizedStructName;
-	FName ImportInfoStructName;
 
 	/** Names for tabs. */
 	static FLazyName ImgMediaProcessImagesTabName;
