ImgMedia: Cache editor tab shows current size and max size.
#rnx

diff --git a/Engine/Plugins/Media/ImgMedia/Source/ImgMedia/Public/ImgMediaGlobalCache.h b/Engine/Plugins/Media/ImgMedia/Source/ImgMedia/Public/ImgMediaGlobalCache.h
--- a/Engine/Plugins/Media/ImgMedia/Source/ImgMedia/Public/ImgMediaGlobalCache.h
+++ b/Engine/Plugins/Media/ImgMedia/Source/ImgMedia/Public/ImgMediaGlobalCache.h
@@ -71,6 +71,16 @@ public:
 	 */
 	void EmptyCache();
 
+	/**
+	 * Get the current size (in bytess).
+	 */
+	SIZE_T GetCurrentSize() const { return CurrentSize; }
+
+	/**
+	 * Get the maximum size (in bytes).
+	 */
+	SIZE_T GetMaxSize() const { return MaxSize; }
+
 private:
 
 	/** An entry in the cache. */
diff --git a/Engine/Plugins/Media/ImgMedia/Source/ImgMediaEditor/Private/Widgets/SImgMediaCache.cpp b/Engine/Plugins/Media/ImgMedia/Source/ImgMediaEditor/Private/Widgets/SImgMediaCache.cpp
--- a/Engine/Plugins/Media/ImgMedia/Source/ImgMediaEditor/Private/Widgets/SImgMediaCache.cpp
+++ b/Engine/Plugins/Media/ImgMedia/Source/ImgMediaEditor/Private/Widgets/SImgMediaCache.cpp
@@ -6,6 +6,7 @@
 #include "ImgMediaGlobalCache.h"
 #include "SlateOptMacros.h"
 #include "Widgets/Input/SButton.h"
+#include "Widgets/Layout/SGridPanel.h"
 #include "Widgets/Layout/SScrollBox.h"
 #include "Widgets/Layout/SSplitter.h"
 #include "Widgets/Text/STextBlock.h"
@@ -27,15 +28,48 @@ void SImgMediaCache::Construct(const FArguments& InArgs)
 			.Padding(0, 20, 0, 0)
 			.AutoHeight()
 		
+		// Global cache label.
 		+ SVerticalBox::Slot()
 			.AutoHeight()
 			.Padding(0, 5, 0, 0)
 		[
 			SNew(STextBlock)
 				.Text(LOCTEXT("GlobalCache", "Global Cache"))
-				.Font(IDetailLayoutBuilder::GetDetailFont())
 		]
+
+		// Global cache info.
+		+ SVerticalBox::Slot()
+			.AutoHeight()
+			[
+				SNew(SGridPanel)
+
+				+ SGridPanel::Slot(0, 0)
+					.Padding(2.0f)
+					[
+						SNew(STextBlock)
+						.Text(LOCTEXT("CurrentSize", "Current Size"))
+					]
+				+ SGridPanel::Slot(1, 0)
+					.Padding(2.0f)
+					[
+						SAssignNew(CurrentSizeTextBlock, STextBlock)
+						.Text(FText::GetEmpty())
+					]
+				+ SGridPanel::Slot(0, 1)
+					.Padding(2.0f)
+					[
+						SNew(STextBlock)
+						.Text(LOCTEXT("MaxSize", "Max Size"))
+					]
+				+ SGridPanel::Slot(1, 1)
+					.Padding(2.0f)
+					[
+						SAssignNew(MaxSizeTextBlock, STextBlock)
+						.Text(FText::GetEmpty())
+					]
+			]
 			
+		// Clear global cache.
 		+ SVerticalBox::Slot()
 			.AutoHeight()
 			.Padding(4.0f)
@@ -46,11 +80,33 @@ void SImgMediaCache::Construct(const FArguments& InArgs)
 					.Text(LOCTEXT("ClearGlobalCache", "Clear Global Cache"))
 					.ToolTipText(LOCTEXT("ClearGlobalCacheButtonToolTip", "Clear the global cache."))
 			]
-		
 	];
 }
 END_SLATE_FUNCTION_BUILD_OPTIMIZATION
 
+void SImgMediaCache::Tick(const FGeometry& AllottedGeometry, const double InCurrentTime, const float InDeltaTime)
+{
+	// Call parent.
+	SCompoundWidget::Tick(AllottedGeometry, InCurrentTime, InDeltaTime);
+
+	// Get cache info.
+	SIZE_T CurrentSize = 0;
+	SIZE_T MaxSize = 0;
+	FImgMediaGlobalCache* GlobalCache = IImgMediaModule::GetGlobalCache();
+	if (GlobalCache != nullptr)
+	{
+		CurrentSize = GlobalCache->GetCurrentSize();
+		MaxSize = GlobalCache->GetMaxSize();
+	}
+
+	// Update widgets.
+	FText NumberText;
+	NumberText = ConvertMemorySizeToText(CurrentSize);
+	CurrentSizeTextBlock->SetText(NumberText);
+	NumberText = ConvertMemorySizeToText(MaxSize);
+	MaxSizeTextBlock->SetText(NumberText);
+}
+
 FReply SImgMediaCache::OnClearGlobalCacheClicked()
 {
 	// Clear the cache.
@@ -63,4 +119,24 @@ FReply SImgMediaCache::OnClearGlobalCacheClicked()
 	return FReply::Handled();
 }
 
+FText SImgMediaCache::ConvertMemorySizeToText(SIZE_T Size)
+{
+	FText NumberText;
+
+	if (Size > (100 * 1024 * 1024))
+	{
+		NumberText = FText::Format(LOCTEXT("SizeMB", "{0} MB"), (int32)(Size / (1024 * 1024)));
+	}
+	else if (Size > 100 * 1024)
+	{
+		NumberText = FText::Format(LOCTEXT("SizeKB", "{0} KB"), (int32)(Size / 1024));
+	}
+	else
+	{
+		NumberText = FText::Format(LOCTEXT("SizeB", "{0} B"), (int32)Size);
+	}
+
+	return NumberText;
+}
+
 #undef LOCTEXT_NAMESPACE
diff --git a/Engine/Plugins/Media/ImgMedia/Source/ImgMediaEditor/Private/Widgets/SImgMediaCache.h b/Engine/Plugins/Media/ImgMedia/Source/ImgMediaEditor/Private/Widgets/SImgMediaCache.h
--- a/Engine/Plugins/Media/ImgMedia/Source/ImgMediaEditor/Private/Widgets/SImgMediaCache.h
+++ b/Engine/Plugins/Media/ImgMedia/Source/ImgMediaEditor/Private/Widgets/SImgMediaCache.h
@@ -6,6 +6,8 @@
 #include "Input/Reply.h"
 #include "Widgets/SCompoundWidget.h"
 
+class STextBlock;
+
 /**
  * SImgMediaCache manages caching for image sequences.
  */
@@ -19,8 +21,25 @@ public:
 
 	void Construct(const FArguments& InArgs);
 
+	// SWidget interface
+	virtual void Tick(const FGeometry& AllottedGeometry, const double InCurrentTime, const float InDeltaTime) override;
+
+
 private:
 	/** Called when we click on the clear global cache button. */
 	FReply OnClearGlobalCacheClicked();
 
+	/**
+	 * Converts a number to a text with KB, MB, etc dependnig on the size.
+	 *
+	 * @param Size		Size in bytes.
+	 * 
+	 * @return Text representing the number.
+	 */
+	FText ConvertMemorySizeToText(SIZE_T Size);
+
+	/** Widget to show curent size of cache. */
+	TSharedPtr<STextBlock> CurrentSizeTextBlock;
+	/** Widget to show max size of cache. */
+	TSharedPtr<STextBlock> MaxSizeTextBlock;
 };
