ImgMedia: Ensure that image loader work can also find locally-cached frames.

diff --git a/Engine/Plugins/Media/ImgMedia/Source/ImgMedia/Private/Loader/ImgMediaLoader.cpp b/Engine/Plugins/Media/ImgMedia/Source/ImgMedia/Private/Loader/ImgMediaLoader.cpp
--- a/Engine/Plugins/Media/ImgMedia/Source/ImgMedia/Private/Loader/ImgMediaLoader.cpp
+++ b/Engine/Plugins/Media/ImgMedia/Source/ImgMedia/Private/Loader/ImgMediaLoader.cpp
@@ -602,8 +602,15 @@ IQueuedWork* FImgMediaLoader::GetWork()
 	FImgMediaLoaderWork* Work = (WorkPool.Num() > 0) ? WorkPool.Pop() : new FImgMediaLoaderWork(AsShared(), Reader.ToSharedRef());
 	
 	// Get the existing frame so we can add the mip level to it.
-	TSharedPtr<FImgMediaFrame, ESPMode::ThreadSafe>* ExistingFramePtr;
-	ExistingFramePtr = GlobalCache->FindAndTouch(SequenceName, FrameNumber);
+	const TSharedPtr<FImgMediaFrame, ESPMode::ThreadSafe>* ExistingFramePtr;
+	if (UseGlobalCache)
+	{
+		ExistingFramePtr = GlobalCache->FindAndTouch(SequenceName, FrameNumber);
+	}
+	else
+	{
+		ExistingFramePtr = Frames.FindAndTouch(FrameNumber);
+	}
 	TSharedPtr<FImgMediaFrame, ESPMode::ThreadSafe> ExistingFrame;
 	
 	
