ImgMedia: Fix for obtaining large tile counts from tile filenames, where string sorting fails to preserve numerical order for non-single digit numbers.

#rnx

diff --git a/Engine/Plugins/Media/ImgMedia/Source/ImgMedia/Private/Loader/ImgMediaLoader.cpp b/Engine/Plugins/Media/ImgMedia/Source/ImgMedia/Private/Loader/ImgMediaLoader.cpp
--- a/Engine/Plugins/Media/ImgMedia/Source/ImgMedia/Private/Loader/ImgMediaLoader.cpp
+++ b/Engine/Plugins/Media/ImgMedia/Source/ImgMedia/Private/Loader/ImgMediaLoader.cpp
@@ -996,15 +996,14 @@ void FImgMediaLoader::GetTileInformation()
 	// Do we have tiles?
 	if (bIsTiled)
 	{
-		// Get the last file as that will tell us the number of tiles.
 		FString Path = FPaths::GetPath(ImagePaths[0][0]);
 		TArray<FString> FoundFiles;
 		IFileManager::Get().FindFiles(FoundFiles, *Path, TEXT("*"));
-		FoundFiles.Sort();
-		if (FoundFiles.Num() > 0)
+
+		for (const FString& FilePath : FoundFiles)
 		{
 			// Remove the file extension.
-			Path = FPaths::ChangeExtension(FoundFiles[FoundFiles.Num() - 1], TEXT(""));
+			Path = FPaths::ChangeExtension(FilePath, TEXT(""));
 
 			// Filename should be ????_x0_y0.
 			// Find the last y.
@@ -1015,7 +1014,7 @@ void FImgMediaLoader::GetTileInformation()
 				int32 YLength = Path.Len() - (YIndex + 1);
 				FString NumberString = Path.Right(YLength);
 				// We count from 0, so add 1 to get the number of tiles.
-				NumTilesY = FCString::Atoi(*NumberString) + 1;
+				NumTilesY = FMath::Max(NumTilesY, FCString::Atoi(*NumberString) + 1);
 
 				// Find the last x.
 				int32 XIndex = 0;
@@ -1025,7 +1024,7 @@ void FImgMediaLoader::GetTileInformation()
 					XIndex++;
 					int32 XLength = Path.Len() - XIndex - YLength - 2;
 					NumberString = Path.Mid(XIndex, XLength);
-					NumTilesX = FCString::Atoi(*NumberString) + 1;
+					NumTilesX = FMath::Max(NumTilesX, FCString::Atoi(*NumberString) + 1);
 				}
 			}
 		}
