ImgMedia: Added custom file picker for the process images sequence path so we can save the directory.

#rnx

diff --git a/Engine/Plugins/Media/ImgMedia/Source/ImgMediaEditor/Private/Customizations/ImgMediaFilePathCustomization.cpp b/Engine/Plugins/Media/ImgMedia/Source/ImgMediaEditor/Private/Customizations/ImgMediaFilePathCustomization.cpp
new file mode 100644
--- /dev/null
+++ b/Engine/Plugins/Media/ImgMedia/Source/ImgMediaEditor/Private/Customizations/ImgMediaFilePathCustomization.cpp
@@ -0,0 +1,84 @@
+// Copyright Epic Games, Inc. All Rights Reserved.
+
+#include "ImgMediaFilePathCustomization.h"
+
+#include "DetailWidgetRow.h"
+#include "IDetailChildrenBuilder.h"
+#include "Misc/ConfigCacheIni.h"
+#include "Widgets/Input/SFilePathPicker.h"
+
+#define LOCTEXT_NAMESPACE "FImgMediaFilePathCustomization"
+
+/* IPropertyTypeCustomization interface
+ *****************************************************************************/
+
+TSharedRef<IPropertyTypeCustomization> FImgMediaFilePathCustomization::MakeInstance()
+{
+	return MakeShared<FImgMediaFilePathCustomization>();
+}
+
+void FImgMediaFilePathCustomization::CustomizeHeader(TSharedRef<IPropertyHandle> PropertyHandle,
+	FDetailWidgetRow& HeaderRow, IPropertyTypeCustomizationUtils& CustomizationUtils)
+{
+}
+
+void FImgMediaFilePathCustomization::CustomizeChildren(TSharedRef<IPropertyHandle> StructPropertyHandle,
+	IDetailChildrenBuilder& ChildBuilder, IPropertyTypeCustomizationUtils& CustomizationUtils)
+{
+	PathStringProperty = StructPropertyHandle->GetChildHandle("FilePath");
+
+	// Create a name for this setting from the owner name and the property name.
+	FProperty* Property = StructPropertyHandle->GetProperty();
+	TObjectPtr<UObject> OwnerObject = Property->GetOwnerUObject();
+	ConfigSettingName = Property->GetName();
+	if (OwnerObject != nullptr)
+	{
+		ConfigSettingName = OwnerObject->GetName() + ConfigSettingName;
+	}
+	
+	// Set filter.
+	FString FileTypeFilter = TEXT("All files (*.*)|*.*");
+
+	// Get the last value for this setting.
+	FString PickedPath;
+	GConfig->GetString(TEXT("ImgMedia.Settings"), *ConfigSettingName, PickedPath, GEditorPerProjectIni);
+	PathStringProperty->SetValue(PickedPath);
+
+	// Add file picker.
+	ChildBuilder.AddCustomRow(LOCTEXT("TimeLabel", "Time"))
+		.NameContent()
+		[
+			StructPropertyHandle->CreatePropertyNameWidget()
+		]
+		.ValueContent()
+		.MaxDesiredWidth(0.0f)
+		.MinDesiredWidth(125.0f)
+		[
+			SNew(SFilePathPicker)
+				.BrowseButtonImage(FEditorStyle::GetBrush("PropertyWindow.Button_Ellipsis"))
+				.BrowseButtonStyle(FEditorStyle::Get(), "HoverHintOnly")
+				.BrowseButtonToolTip(LOCTEXT("FileButtonToolTipText", "Choose a file from this computer"))
+				.BrowseDirectory(FPaths::GetPath(PickedPath))
+				.BrowseTitle(LOCTEXT("PropertyEditorTitle", "File picker..."))
+				.FilePath(this, &FImgMediaFilePathCustomization::HandleFilePathPickerFilePath)
+				.FileTypeFilter(FileTypeFilter)
+				.OnPathPicked(this, &FImgMediaFilePathCustomization::HandleFilePathPickerPathPicked)
+		];
+}
+
+FString FImgMediaFilePathCustomization::HandleFilePathPickerFilePath() const
+{
+	FString FilePath;
+	PathStringProperty->GetValue(FilePath);
+
+	return FilePath;
+}
+
+void FImgMediaFilePathCustomization::HandleFilePathPickerPathPicked(const FString& PickedPath)
+{
+	PathStringProperty->SetValue(PickedPath);
+	GConfig->SetString(TEXT("ImgMedia.Settings"), *ConfigSettingName, *PickedPath, GEditorPerProjectIni);
+}
+
+
+#undef LOCTEXT_NAMESPACE
diff --git a/Engine/Plugins/Media/ImgMedia/Source/ImgMediaEditor/Private/Customizations/ImgMediaFilePathCustomization.h b/Engine/Plugins/Media/ImgMedia/Source/ImgMediaEditor/Private/Customizations/ImgMediaFilePathCustomization.h
new file mode 100644
--- /dev/null
+++ b/Engine/Plugins/Media/ImgMedia/Source/ImgMediaEditor/Private/Customizations/ImgMediaFilePathCustomization.h
@@ -0,0 +1,42 @@
+// Copyright Epic Games, Inc. All Rights Reserved.
+
+#pragma once
+
+#include "CoreTypes.h"
+#include "IPropertyTypeCustomization.h"
+
+class IPropertyHandle;
+
+/**
+ * Implements a customization for the FFilePath class.
+ * 
+ * This will preserve the file path when closing.
+ */
+class FImgMediaFilePathCustomization : public IPropertyTypeCustomization
+{
+public:
+	static TSharedRef<IPropertyTypeCustomization> MakeInstance();
+
+	virtual void CustomizeHeader(TSharedRef<IPropertyHandle> PropertyHandle, FDetailWidgetRow& HeaderRow,
+		IPropertyTypeCustomizationUtils& CustomizationUtils) override;
+	virtual void CustomizeChildren(TSharedRef<IPropertyHandle> PropertyHandle, IDetailChildrenBuilder& ChildBuilder,
+		IPropertyTypeCustomizationUtils& CustomizationUtils) override;
+
+private:
+	/**
+	 * Gets the file path.
+	 */
+	FString HandleFilePathPickerFilePath() const;
+
+	/**
+	 * Called when a path is selected.
+	 */
+	void HandleFilePathPickerPathPicked(const FString& PickedPath);
+
+	/** Pointer to the string that will be seet when changing the path */
+	TSharedPtr<IPropertyHandle> PathStringProperty;
+
+	/** Name to save the file path to. */
+	FString ConfigSettingName;
+};
+
diff --git a/Engine/Plugins/Media/ImgMedia/Source/ImgMediaEditor/Private/Widgets/SImgMediaProcessImages.cpp b/Engine/Plugins/Media/ImgMedia/Source/ImgMediaEditor/Private/Widgets/SImgMediaProcessImages.cpp
--- a/Engine/Plugins/Media/ImgMedia/Source/ImgMediaEditor/Private/Widgets/SImgMediaProcessImages.cpp
+++ b/Engine/Plugins/Media/ImgMedia/Source/ImgMediaEditor/Private/Widgets/SImgMediaProcessImages.cpp
@@ -3,6 +3,7 @@
 #include "Widgets/SImgMediaProcessImages.h"
 
 #include "Async/Async.h"
+#include "Customizations/ImgMediaFilePathCustomization.h"
 #include "EditorStyleSet.h"
 #include "Framework/Notifications/NotificationManager.h"
 #include "HAL/PlatformFileManager.h"
@@ -72,6 +73,8 @@ void SImgMediaProcessImages::Construct(const FArguments& InArgs)
 	DetailsViewArgs.bAllowSearch = false;
 	DetailsViewArgs.NameAreaSettings = FDetailsViewArgs::HideNameArea;
 	DetailsView = PropertyEditorModule.CreateDetailView(DetailsViewArgs);
+	DetailsView->RegisterInstancedCustomPropertyTypeLayout(FName(TEXT("FilePath")),
+		FOnGetPropertyTypeCustomizationInstance::CreateStatic(FImgMediaFilePathCustomization::MakeInstance));
 	DetailsView->SetObject(Options.Get());
 
 	DetailsViewBox->SetContent(DetailsView->AsShared());
