ImgMedia: Added option to process images tab to use our custom format.
#rnx

diff --git a/Engine/Plugins/Media/ImgMedia/Source/ImgMediaEditor/Private/Widgets/ImgMediaProcessImagesOptions.h b/Engine/Plugins/Media/ImgMedia/Source/ImgMediaEditor/Private/Widgets/ImgMediaProcessImagesOptions.h
--- a/Engine/Plugins/Media/ImgMedia/Source/ImgMediaEditor/Private/Widgets/ImgMediaProcessImagesOptions.h
+++ b/Engine/Plugins/Media/ImgMedia/Source/ImgMediaEditor/Private/Widgets/ImgMediaProcessImagesOptions.h
@@ -22,6 +22,10 @@ public:
 	/** The directory to output the processed image sequence files to. */
 	UPROPERTY(EditAnywhere, Category = Sequence)
 	FDirectoryPath OutputPath;
+
+	/** If checked, then use our custom format with a single file for all tiles. */
+	UPROPERTY(EditAnywhere, Category = Sequence)
+	bool bUseCustomFormat;
 	
 	/** Width of a tile in pixels. If 0, then do not make tiles. */
 	UPROPERTY(EditAnywhere, Transient, Category = Tiles)
diff --git a/Engine/Plugins/Media/ImgMedia/Source/ImgMediaEditor/Private/Widgets/SImgMediaProcessImages.cpp b/Engine/Plugins/Media/ImgMedia/Source/ImgMediaEditor/Private/Widgets/SImgMediaProcessImages.cpp
--- a/Engine/Plugins/Media/ImgMedia/Source/ImgMediaEditor/Private/Widgets/SImgMediaProcessImages.cpp
+++ b/Engine/Plugins/Media/ImgMedia/Source/ImgMediaEditor/Private/Widgets/SImgMediaProcessImages.cpp
@@ -99,6 +99,7 @@ FReply SImgMediaProcessImages::OnProcessImagesClicked()
 
 void SImgMediaProcessImages::ProcessAllImages(TSharedPtr<SNotificationItem> ConfirmNotification)
 {
+	bool bUseCustomFormat = Options->bUseCustomFormat;
 	int32 InTileWidth = Options->TileSizeX;
 	int32 InTileHeight = Options->TileSizeY;
 
@@ -166,8 +167,16 @@ void SImgMediaProcessImages::ProcessAllImages(TSharedPtr<SNotificationItem> Conf
 				}
 
 				// Import this image.
-				FString Name = FPaths::ChangeExtension(FPaths::Combine(OutPath, FileName), TEXT(""));
-				ProcessImage(ImageWrapper, InTileWidth, InTileHeight, Name, Ext);
+				FString Name = FPaths::Combine(OutPath, FileName);
+				if (bUseCustomFormat)
+				{
+					ProcessImageCustom(ImageWrapper, InTileWidth, InTileHeight, Name);
+				}
+				else
+				{
+					Name = FPaths::ChangeExtension(Name, TEXT(""));
+					ProcessImage(ImageWrapper, InTileWidth, InTileHeight, Name, Ext);
+				}
 			}
 		}
 	}
@@ -248,4 +257,9 @@ void SImgMediaProcessImages::ProcessImage(TSharedPtr<IImageWrapper>& InImageWrap
 	}
 }
 
+void SImgMediaProcessImages::ProcessImageCustom(TSharedPtr<IImageWrapper>& InImageWrapper,
+	int32 InTileWidth, int32 InTileHeight, const FString& InName)
+{
+}
+
 #undef LOCTEXT_NAMESPACE
diff --git a/Engine/Plugins/Media/ImgMedia/Source/ImgMediaEditor/Private/Widgets/SImgMediaProcessImages.h b/Engine/Plugins/Media/ImgMedia/Source/ImgMediaEditor/Private/Widgets/SImgMediaProcessImages.h
--- a/Engine/Plugins/Media/ImgMedia/Source/ImgMediaEditor/Private/Widgets/SImgMediaProcessImages.h
+++ b/Engine/Plugins/Media/ImgMedia/Source/ImgMediaEditor/Private/Widgets/SImgMediaProcessImages.h
@@ -50,6 +50,19 @@ private:
 	void ProcessImage(TSharedPtr<IImageWrapper>& InImageWrapper,
 		int32 InTileWidth, int32 InTileHeight, const FString& InName, const FString& FileExtension);
 
+	/**
+	 * Processess a single image and writes out a file.
+	 * Tiles and mips may be generated.
+	 * This does NOT run on the game thread.
+	 *
+	 * @param InImageWrapper	ImageWrapper to read/write the image.
+	 * @param InTileWidth		Desired width of tiles.
+	 * @param InTileHeight		Desired height of tiles.
+	 * @param InName			Full path and name of file to write.
+	 */
+	void ProcessImageCustom(TSharedPtr<IImageWrapper>& InImageWrapper,
+		int32 InTileWidth, int32 InTileHeight, const FString& InName);
+
 	/** Holds our details view. */
 	TSharedPtr<class IDetailsView> DetailsView;
 	/** Object that holds our options. */
