ImgMedia: Added writing out tiled EXRs to OpenExrWrapper.

#rnx

diff --git a/Engine/Plugins/Media/ImgMedia/Source/OpenExrWrapper/Private/OpenExrWrapper.cpp b/Engine/Plugins/Media/ImgMedia/Source/OpenExrWrapper/Private/OpenExrWrapper.cpp
--- a/Engine/Plugins/Media/ImgMedia/Source/OpenExrWrapper/Private/OpenExrWrapper.cpp
+++ b/Engine/Plugins/Media/ImgMedia/Source/OpenExrWrapper/Private/OpenExrWrapper.cpp
@@ -11,8 +11,10 @@ THIRD_PARTY_INCLUDES_START
 	#include "Imath/ImathBox.h"
 	#include "OpenEXR/ImfCompressionAttribute.h"
 	#include "OpenEXR/ImfHeader.h"
+	#include "OpenEXR/ImfIntAttribute.h"
 	#include "OpenEXR/ImfRgbaFile.h"
 	#include "OpenEXR/ImfStandardAttributes.h"
+	#include "OpenEXR/ImfTiledRgbaFile.h"
 THIRD_PARTY_INCLUDES_END
 PRAGMA_DEFAULT_VISIBILITY_END
 
@@ -212,5 +214,146 @@ void FRgbaInputFile::SetFrameBuffer(void* Buffer, const FIntPoint& BufferDim)
 	((Imf::RgbaInputFile*)InputFile)->setFrameBuffer((Imf::Rgba*)Buffer - Win.min.x - Win.min.y * BufferDim.X, 1, BufferDim.X);
 }
 
+FTiledRgbaOutputFile::FTiledRgbaOutputFile(
+	const FIntPoint& DisplayWindowMin,
+	const FIntPoint& DisplayWindowMax,
+	const FIntPoint& DataWindowMin,
+	const FIntPoint& DataWindowMax)
+{
+	OutputFile = nullptr;
+
+	Imath::Box2i EXRDisplayWindow = Imath::Box2i(Imath::V2i(DisplayWindowMin.X, DisplayWindowMin.Y),
+		Imath::V2i(DisplayWindowMax.X, DisplayWindowMax.Y));
+	Imath::Box2i EXRDataWindow = Imath::Box2i(Imath::V2i(DataWindowMin.X, DataWindowMin.Y),
+		Imath::V2i(DataWindowMax.X, DataWindowMax.Y));
+	
+	Header = new Imf::Header(EXRDisplayWindow, EXRDataWindow, 1, IMATH_NAMESPACE::V2f(0, 0), 1, Imf::INCREASING_Y,
+		Imf::NO_COMPRESSION);
+}
+
+FTiledRgbaOutputFile::~FTiledRgbaOutputFile()
+{
+	if (Header != nullptr)
+	{
+		delete (Imf::Header*)Header;
+	}
+	if (OutputFile != nullptr)
+	{
+		delete (Imf::TiledRgbaOutputFile*)OutputFile;
+	}
+}
+
+void FTiledRgbaOutputFile::AddIntAttribute(const FString& Name, int32 Value)
+{
+	// Make sure we don't have an output file yet.
+	if (OutputFile == nullptr)
+	{
+		((Imf::Header*)Header)->insert(std::string(TCHAR_TO_ANSI(*Name)), Imf::IntAttribute(Value));
+	}
+	else
+	{
+		UE_LOG(LogOpenEXRWrapper, Error, TEXT("Attribute %s added after calling CreateOutputFile."),
+			*Name);
+	}
+}
+
+void FTiledRgbaOutputFile::CreateOutputFile(const FString& FilePath, 
+	int32 TileWidth, int32 TileHeight, int32 NumChannels, bool bIsMipsEnabled)
+{
+	if (OutputFile == nullptr)
+	{
+		try
+		{
+			// Get channels.
+			Imf::RgbaChannels Channels;
+			switch (NumChannels)
+			{
+			case 1:
+				Channels = Imf::WRITE_R;
+				break;
+			case 2:
+				Channels = Imf::WRITE_YC;
+				break;
+			case 3:
+				Channels = Imf::WRITE_RGB;
+				break;
+			case 4:
+				Channels = Imf::WRITE_RGBA;
+				break;
+			default:
+				UE_LOG(LogOpenEXRWrapper, Error, TEXT("Unsupported number of channels %d"),
+					NumChannels);
+				Channels = Imf::WRITE_RGBA;
+				break;
+			}
+
+			// Create output file.
+			OutputFile = new Imf::TiledRgbaOutputFile(TCHAR_TO_ANSI(*FilePath),
+				*((Imf::Header*)Header),
+				Channels,
+				TileWidth, TileHeight,
+				bIsMipsEnabled ? Imf::MIPMAP_LEVELS : Imf::ONE_LEVEL,
+				Imf::ROUND_DOWN);
+		}
+		catch (std::exception const& Exception)
+		{
+			UE_LOG(LogOpenEXRWrapper, Error, TEXT("Cannot write EXR file: %s (%s)"),
+				*FilePath, StringCast<TCHAR>(Exception.what()).Get());
+		}
+	}
+	else
+	{
+		UE_LOG(LogOpenEXRWrapper, Error,
+			TEXT("Cannot create output file as it has already been created."));
+	}
+}
+
+int32 FTiledRgbaOutputFile::GetNumberOfMipLevels()
+{
+	if (OutputFile != nullptr)
+	{
+		return ((Imf::TiledRgbaOutputFile*)OutputFile)->numLevels();
+	}
+	else
+	{
+		UE_LOG(LogOpenEXRWrapper, Error,
+			TEXT("GetNumberOfMipLevels failed: CreateOutputFile has not been called yet."));
+		return 0;
+	}
+}
+
+void FTiledRgbaOutputFile::SetFrameBuffer(void* Buffer, const FIntPoint& Stride)
+{
+	if (OutputFile != nullptr)
+	{
+		((Imf::TiledRgbaOutputFile*)OutputFile)->setFrameBuffer((Imf::Rgba*)Buffer, Stride.X, Stride.Y);
+	}
+	else
+	{
+		UE_LOG(LogOpenEXRWrapper, Error,
+			TEXT("SetFrameBuffer failed: CreateOutputFile has not been called yet."));
+	}
+}
+
+void FTiledRgbaOutputFile::WriteTile(int32 TileX, int32 TileY, int32 MipLevel)
+{
+	if (OutputFile != nullptr)
+	{
+		try
+		{
+			((Imf::TiledRgbaOutputFile*)OutputFile)->writeTile(TileX, TileY, MipLevel);
+		}
+		catch (std::exception const& Exception)
+		{
+			UE_LOG(LogOpenEXRWrapper, Error, TEXT("Cannot write EXR file: %s"),
+				StringCast<TCHAR>(Exception.what()).Get());
+		}
+	}
+	else
+	{
+		UE_LOG(LogOpenEXRWrapper, Error,
+			TEXT("WriteTile failed: CreateOutputFile has not been called yet."));
+	}
+}
 
 IMPLEMENT_MODULE(FDefaultModuleImpl, OpenExrWrapper);
diff --git a/Engine/Plugins/Media/ImgMedia/Source/OpenExrWrapper/Public/OpenExrWrapper.h b/Engine/Plugins/Media/ImgMedia/Source/OpenExrWrapper/Public/OpenExrWrapper.h
--- a/Engine/Plugins/Media/ImgMedia/Source/OpenExrWrapper/Public/OpenExrWrapper.h
+++ b/Engine/Plugins/Media/ImgMedia/Source/OpenExrWrapper/Public/OpenExrWrapper.h
@@ -42,3 +42,82 @@ private:
 
 	void* InputFile;
 };
+
+/**
+ * Use this to write out tiled EXR images.
+ * 
+ * Add any attributes after construction.
+ * Then you can call CreateOutputFile.
+ * After that, you can write out data to the file.
+ */
+class OPENEXRWRAPPER_API FTiledRgbaOutputFile
+{
+public:
+	/**
+	 * Constructor.
+	 * 
+	 * @param DisplayWindowMin		Normally (0, 0).
+	 * @param DisplayWindowMax		Normally (width - 1, height - 1).
+	 * @param DataWindowMin			Normally (0, 0).
+	 * @param DataWindowMax			Normally (width - 1, height - 1).
+	 */
+	FTiledRgbaOutputFile(
+		const FIntPoint& DisplayWindowMin,
+		const FIntPoint& DisplayWindowMax,
+		const FIntPoint& DataWindowMin,
+		const FIntPoint& DataWindowMax);
+
+	/**
+	 * Destructor.
+	 */
+	~FTiledRgbaOutputFile();
+
+	/**
+	 * Call this to add an attribute to the EXR file.
+	 * This MUST be called before CreateOutputFile.
+	 * 
+	 * @param Name		Name for this attribute.
+	 * @param Value		Value for this attribute.
+	 */
+	void AddIntAttribute(const FString& Name, int32 Value);
+
+	/**
+	 * Call this after adding any attributes BUT before doing anything else.
+	 *
+	 * @param FilePath			Filename to save to.
+	 * @param TileWidth			Width of a tile.
+	 * @param TileHeight		Height of a tile.
+	 * @param NumChannels		Number of channels out write out.
+	 * @param bIsMipsEnabled	True to enable mip mapping.
+	 */
+	void CreateOutputFile(const FString& FilePath,
+		int32 TileWidth, int32 TileHeight, int32 NumChannels, bool bIsMipsEnabled);
+
+	/**
+	 * Call this to get the number of mipmap levels.
+	 */
+	int32 GetNumberOfMipLevels();
+
+	/**
+	 * Call this prior to WriteTile to set where the data is coming from.
+	 * 
+	 * @param Buffer		Source of data.
+	 * @param Stride		A pixels location is calculated by Buffer + x * StrideX + y * StrideY.
+	 */
+	void SetFrameBuffer(void* Buffer, const FIntPoint& Stride);
+
+	/**
+	 * Call this to write data to a specific tile.
+	 *
+	 * @param TileX			X coordinate of tile.
+	 * @param TileY			Y coordinate of tile.
+	 * @param MipLevel		Mipmap level of tile.
+	 */
+	void WriteTile(int32 TileX, int32 TileY, int32 MipLevel);
+
+private:
+	/** Stores the EXR header object. */
+	void* Header;
+	/** Stores the EXR object. */
+	void* OutputFile;
+};
