ImgMedia: ProcessImages tab can output native tiled EXRs as its custom format.

#rnx

diff --git a/Engine/Plugins/Media/ImgMedia/Source/ImgMediaEditor/ImgMediaEditor.Build.cs b/Engine/Plugins/Media/ImgMedia/Source/ImgMediaEditor/ImgMediaEditor.Build.cs
--- a/Engine/Plugins/Media/ImgMedia/Source/ImgMediaEditor/ImgMediaEditor.Build.cs
+++ b/Engine/Plugins/Media/ImgMedia/Source/ImgMediaEditor/ImgMediaEditor.Build.cs
@@ -39,6 +39,21 @@ namespace UnrealBuildTool.Rules
 					"Engine",
 				});
 			}
+
+			bool bLinuxEnabled = Target.Platform == UnrealTargetPlatform.Linux && Target.Architecture.StartsWith("x86_64");
+
+			// Add EXR support.
+			if ((Target.Platform == UnrealTargetPlatform.Mac) ||
+				Target.Platform.IsInGroup(UnrealPlatformGroup.Windows) ||
+				bLinuxEnabled)
+			{
+				PrivateDependencyModuleNames.Add("OpenExrWrapper");
+				PublicDefinitions.Add("IMGMEDIAEDITOR_EXR_SUPPORTED_PLATFORM=1");
+			}
+			else
+			{
+				PublicDefinitions.Add("IMGMEDIAEDITOR_EXR_SUPPORTED_PLATFORM=0");
+			}
 		}
 	}
 }
diff --git a/Engine/Plugins/Media/ImgMedia/Source/ImgMediaEditor/Private/Widgets/SImgMediaProcessImages.cpp b/Engine/Plugins/Media/ImgMedia/Source/ImgMediaEditor/Private/Widgets/SImgMediaProcessImages.cpp
--- a/Engine/Plugins/Media/ImgMedia/Source/ImgMediaEditor/Private/Widgets/SImgMediaProcessImages.cpp
+++ b/Engine/Plugins/Media/ImgMedia/Source/ImgMediaEditor/Private/Widgets/SImgMediaProcessImages.cpp
@@ -24,6 +24,10 @@
 #include "Widgets/Notifications/SNotificationList.h"
 #include "Widgets/Text/STextBlock.h"
 
+#if IMGMEDIAEDITOR_EXR_SUPPORTED_PLATFORM
+#include "OpenExrWrapper.h"
+#endif // IMGMEDIAEDITOR_EXR_SUPPORTED_PLATFORM
+
 #define LOCTEXT_NAMESPACE "ImgMediaProcessImages"
 
 SImgMediaProcessImages::~SImgMediaProcessImages()
@@ -260,6 +264,41 @@ void SImgMediaProcessImages::ProcessImage(TSharedPtr<IImageWrapper>& InImageWrap
 void SImgMediaProcessImages::ProcessImageCustom(TSharedPtr<IImageWrapper>& InImageWrapper,
 	int32 InTileWidth, int32 InTileHeight, const FString& InName)
 {
+#if IMGMEDIAEDITOR_EXR_SUPPORTED_PLATFORM
+	// Get image data.
+	ERGBFormat Format = InImageWrapper->GetFormat();
+	int32 Width = InImageWrapper->GetWidth();
+	int32 Height = InImageWrapper->GetHeight();
+	int32 BitDepth = InImageWrapper->GetBitDepth();
+	TArray64<uint8> RawData;
+	InImageWrapper->GetRaw(Format, BitDepth, RawData);
+
+	int32 NumTilesX = InTileWidth > 0 ? Width / InTileWidth : 1;
+	int32 NumTilesY = InTileHeight > 0 ? Height / InTileHeight : 1;
+	int32 TileWidth = Width / NumTilesX;
+	int32 TileHeight = Height / NumTilesY;
+	int32 BytesPerPixel = RawData.Num() / (Width * Height);
+
+	// Create tiled exr file.
+	FTiledRgbaOutputFile OutFile(FIntPoint(0, 0), FIntPoint(Width - 1, Height - 1),
+		FIntPoint(0, 0), FIntPoint(Width - 1, Height - 1));
+	OutFile.CreateOutputFile(InName, TileWidth, TileHeight, 4, false);
+	
+	FIntPoint Stride(1, Width);
+	
+	// Loop over y tiles.
+	for (int TileY = 0; TileY < NumTilesY; ++TileY)
+	{
+		// Loop over x tiles.
+		for (int TileX = 0; TileX < NumTilesX; ++TileX)
+		{
+			OutFile.SetFrameBuffer(RawData.GetData(), Stride);
+			OutFile.WriteTile(TileX, TileY, 0);
+		}
+	}
+#else // IMGMEDIAEDITOR_EXR_SUPPORTED_PLATFORM
+	UE_LOG(LogImgMediaEditor, Error, TEXT("EXR not supported on this platform."));
+#endif // IMGMEDIAEDITOR_EXR_SUPPORTED_PLATFORM
 }
 
 #undef LOCTEXT_NAMESPACE
