ImgMedia: Can get int attributes from OpenExrWrapper.
#rnx

diff --git a/Engine/Plugins/Media/ImgMedia/Source/OpenExrWrapper/Private/OpenExrWrapper.cpp b/Engine/Plugins/Media/ImgMedia/Source/OpenExrWrapper/Private/OpenExrWrapper.cpp
--- a/Engine/Plugins/Media/ImgMedia/Source/OpenExrWrapper/Private/OpenExrWrapper.cpp
+++ b/Engine/Plugins/Media/ImgMedia/Source/OpenExrWrapper/Private/OpenExrWrapper.cpp
@@ -214,6 +214,26 @@ void FRgbaInputFile::SetFrameBuffer(void* Buffer, const FIntPoint& BufferDim)
 	((Imf::RgbaInputFile*)InputFile)->setFrameBuffer((Imf::Rgba*)Buffer - Win.min.x - Win.min.y * BufferDim.X, 1, BufferDim.X);
 }
 
+bool FRgbaInputFile::GetIntAttribute(const FString& Name, int32& Value)
+{
+	bool bIsAttributeFound = false;
+
+	if (InputFile != nullptr)
+	{
+		const Imf::IntAttribute* Attribute =
+			((Imf::RgbaInputFile*)InputFile)->header().
+			findTypedAttribute<Imf::IntAttribute>(std::string(TCHAR_TO_ANSI(*Name)));
+
+		if (Attribute != nullptr)
+		{
+			Value = Attribute->value();
+			bIsAttributeFound = true;
+		}
+	}
+
+	return bIsAttributeFound;
+}
+
 FTiledRgbaOutputFile::FTiledRgbaOutputFile(
 	const FIntPoint& DisplayWindowMin,
 	const FIntPoint& DisplayWindowMax,
diff --git a/Engine/Plugins/Media/ImgMedia/Source/OpenExrWrapper/Public/OpenExrWrapper.h b/Engine/Plugins/Media/ImgMedia/Source/OpenExrWrapper/Public/OpenExrWrapper.h
--- a/Engine/Plugins/Media/ImgMedia/Source/OpenExrWrapper/Public/OpenExrWrapper.h
+++ b/Engine/Plugins/Media/ImgMedia/Source/OpenExrWrapper/Public/OpenExrWrapper.h
@@ -38,6 +38,16 @@ public:
 	void ReadPixels(int32 StartY, int32 EndY);
 	void SetFrameBuffer(void* Buffer, const FIntPoint& Stride);
 
+	/**
+	 * Get an attribute from the image.
+	 *
+	 * @param Name		Name of attribute.
+	 * @param Value		Will be set to the value of the attribute if the attribute is found.
+	 *					Will NOT be set if the attribute is not found.
+	 * @return			True if the attribute was found, false otherwise.
+	 */
+	bool GetIntAttribute(const FString& Name, int32& Value);
+
 private:
 
 	void* InputFile;
