Image Media Exr Reader GPU:
- Added ability to read individual tile files and swizzle those into single texture.
- Exr swizzler now operates only on a region instead of full texture.

diff --git a/Engine/Plugins/Media/ImgMedia/Shaders/Private/ExrSwizzler.usf b/Engine/Plugins/Media/ImgMedia/Shaders/Private/ExrSwizzler.usf
--- a/Engine/Plugins/Media/ImgMedia/Shaders/Private/ExrSwizzler.usf
+++ b/Engine/Plugins/Media/ImgMedia/Shaders/Private/ExrSwizzler.usf
@@ -5,11 +5,12 @@
 /** The unswizzled buffer packed into 4 byte buffer stride. */
 StructuredBuffer<uint> UnswizzledBuffer;
 
-/** Texture width that the buffer is supposed to be unpacked to. */
-int TextureWidth;
+/** Texture size that the buffer is supposed to be unpacked to. */
+int2 TextureSize;
+
+/** Size of tiles that represent this texture. */
+int2 TileSize;
 
-/** Texture height that the buffer is supposed to be unpacked to. */
-int TextureHeight;
 
 /**
 * This function extracts the 16bit (float) data stored in uint (32 bit)
@@ -27,10 +28,10 @@ float UintContainingFloatToFloat(in uint Value)
 }
 
 /** Unpacks data from the provided structured buffer. */
-int GetValueFromTheBuffer(int Offset, int TextureWidth, int TextureHeight, int ChannelIndex)
+int GetValueFromTheBuffer(int Offset, int2 TextureSize, int ChannelIndex)
 {
 	// the location as if it would've been if the buffer was unpacked.
-	int UnpackedLocation = Offset + TextureWidth * ChannelIndex;
+	int UnpackedLocation = Offset + TextureSize.x * ChannelIndex;
 
 	// Byte offset within the 4 byte packed value.
 	int Remainder = UnpackedLocation % 2;
@@ -55,14 +56,8 @@ void MainVS(
 	DrawRectangle(InPosition, InTexCoord, OutPosition, OutUVAndScreenPos);
 }
 
-
-/** Pixel Shader. */
-half4 MainPS(noperspective float4 UVAndScreenPos : TEXCOORD0) : SV_Target0
+void GetStartScanlinePosition(in int2 PixelPos, out float2 SampleSize, out int StartPosition)
 {
-	float2 UV = UVAndScreenPos.xy;
-	int xCoord = UV.x * TextureWidth;
-	int yCoord = UV.y * TextureHeight;
-
 	// Depending on the EXR file it could contain multi-channel data.
 #if NUM_CHANNELS == 1
 	int Channels = 1;
@@ -73,27 +68,63 @@ half4 MainPS(noperspective float4 UVAndScreenPos : TEXCOORD0) : SV_Target0
 #else
 	int Channels = 4;
 #endif
+
+#if SWIZZLE_TILES
+	SampleSize.x = TileSize.x;
+	SampleSize.y = TileSize.y;
+
+	int xCoord = (PixelPos.x) % TileSize.x;
+	int yCoord = (PixelPos.y) % TileSize.y;
+
+	int TileX = floor((PixelPos.x) / TileSize.x);
+	int TileY = floor((PixelPos.y) / TileSize.y);
+
+	int MaxTilesX = TextureSize.x / TileSize.x;
+	int MaxTilesY = TextureSize.y / TileSize.y;
+
+	// Current scanline position.
+	int TileBufferStride = TileSize.x * TileSize.y * Channels + TileSize.y * 4;
+	int PreviousTilesOffset = (MaxTilesX * (TileY)+TileX) * TileBufferStride;
+	StartPosition = PreviousTilesOffset + yCoord * TileSize.x * Channels + xCoord + (yCoord + 1) * 4;
+#else
+	SampleSize.x = TextureSize.x;
+	SampleSize.y = TextureSize.y;
+
 	// Current scanline position.
-	int Offset = yCoord * TextureWidth * Channels + xCoord + (yCoord + 1) * 4;
+	StartPosition = PixelPos.y * TextureSize.x * Channels + PixelPos.x + (PixelPos.y + 1) * 4;
+#endif
+}
+
+/** Pixel Shader. */
+half4 MainPS(noperspective float4 UVAndScreenPos : TEXCOORD0) : SV_Target0
+{
+	float2 UV = UVAndScreenPos.xy;
+
+	float2 SampleSize;
+
+	int2 PixelPos = UV * TextureSize;
+	int StartBufferPosition;
+
+	GetStartScanlinePosition(PixelPos, SampleSize, StartBufferPosition);
+
 	int ChannelIndex = 0;
 #if NUM_CHANNELS == 4
-	half A = UintContainingFloatToFloat(GetValueFromTheBuffer(Offset, TextureWidth, TextureHeight, ChannelIndex++));
+	half A = UintContainingFloatToFloat(GetValueFromTheBuffer(StartBufferPosition, SampleSize, ChannelIndex++));
 #else
 	half A = 1.;
 #endif
 
-	half B = UintContainingFloatToFloat(GetValueFromTheBuffer(Offset, TextureWidth, TextureHeight, ChannelIndex++));
+	half B = UintContainingFloatToFloat(GetValueFromTheBuffer(StartBufferPosition, SampleSize, ChannelIndex++));
 #if NUM_CHANNELS <= 1
 	half G = B;
 #else
-	half G = UintContainingFloatToFloat(GetValueFromTheBuffer(Offset, TextureWidth, TextureHeight, ChannelIndex++));
+	half G = UintContainingFloatToFloat(GetValueFromTheBuffer(StartBufferPosition, SampleSize, ChannelIndex++));
 #endif
 #if NUM_CHANNELS <= 2
 	half R = B;
 #else
-	half R = UintContainingFloatToFloat(GetValueFromTheBuffer(Offset, TextureWidth, TextureHeight, ChannelIndex++));
+	half R = UintContainingFloatToFloat(GetValueFromTheBuffer(StartBufferPosition, SampleSize, ChannelIndex++));
 #endif
-
 	half3 Rgb = half3(R, G, B);
 
 	return half4(Rgb, A);
diff --git a/Engine/Plugins/Media/ImgMedia/Source/ExrReaderGpu/Private/ExrReaderGpu.cpp b/Engine/Plugins/Media/ImgMedia/Source/ExrReaderGpu/Private/ExrReaderGpu.cpp
--- a/Engine/Plugins/Media/ImgMedia/Source/ExrReaderGpu/Private/ExrReaderGpu.cpp
+++ b/Engine/Plugins/Media/ImgMedia/Source/ExrReaderGpu/Private/ExrReaderGpu.cpp
@@ -297,13 +297,13 @@ bool FExrReader::ReadExrImageChunk(void* Buffer, int32 ChunkSize)
 		UE_LOG(LogExrReaderGpu, Error, TEXT("Buffer provided is invalid. Please provide a valid buffer. "));
 		return false;
 	}
-
-	size_t Result = fread(Buffer, ChunkSize, 1, FileHandle);
-	if (Result != 1)
+	size_t NumElements = 1;
+	size_t NumElementsRead = fread(Buffer, ChunkSize, NumElements, FileHandle);
+	if (NumElementsRead != NumElements)
 	{
 		UE_LOG(LogExrReaderGpu, Error, TEXT("Issue reading EXR chunk. "));
 	}
-	return Result == 1;
+	return NumElementsRead == NumElements;
 }
 
 bool FExrReader::CloseExrFile()
diff --git a/Engine/Plugins/Media/ImgMedia/Source/ExrReaderGpu/Public/ExrSwizzlingShader.h b/Engine/Plugins/Media/ImgMedia/Source/ExrReaderGpu/Public/ExrSwizzlingShader.h
--- a/Engine/Plugins/Media/ImgMedia/Source/ExrReaderGpu/Public/ExrSwizzlingShader.h
+++ b/Engine/Plugins/Media/ImgMedia/Source/ExrReaderGpu/Public/ExrSwizzlingShader.h
@@ -41,7 +41,8 @@ class EXRREADERGPU_API FExrSwizzlePS : public FGlobalShader
 
 	/** If the provided buffer is RGBA the shader would work slightly differently to RGB. */
 	class FRgbaSwizzle : SHADER_PERMUTATION_INT("NUM_CHANNELS", 5);
-	using FPermutationDomain = TShaderPermutationDomain<FRgbaSwizzle>;
+	class FSwizzleTiles : SHADER_PERMUTATION_BOOL("SWIZZLE_TILES");
+	using FPermutationDomain = TShaderPermutationDomain<FRgbaSwizzle, FSwizzleTiles>;
 
 	static bool ShouldCompilePermutation(const FGlobalShaderPermutationParameters& Parameters) 
 	{
@@ -50,8 +51,8 @@ class EXRREADERGPU_API FExrSwizzlePS : public FGlobalShader
 
 	BEGIN_SHADER_PARAMETER_STRUCT(FParameters, )
 		SHADER_PARAMETER_SRV(StructuredBuffer<uint>, UnswizzledBuffer)
-		SHADER_PARAMETER(int, TextureWidth)
-		SHADER_PARAMETER(int, TextureHeight)
+		SHADER_PARAMETER(FIntPoint, TextureSize)
+		SHADER_PARAMETER(FIntPoint, TileSize)
 	END_SHADER_PARAMETER_STRUCT()
 };
 #endif
\ No newline at end of file
diff --git a/Engine/Plugins/Media/ImgMedia/Source/ImgMedia/Private/Readers/ExrImgMediaReaderGpu.cpp b/Engine/Plugins/Media/ImgMedia/Source/ImgMedia/Private/Readers/ExrImgMediaReaderGpu.cpp
--- a/Engine/Plugins/Media/ImgMedia/Source/ImgMedia/Private/Readers/ExrImgMediaReaderGpu.cpp
+++ b/Engine/Plugins/Media/ImgMedia/Source/ImgMedia/Private/Readers/ExrImgMediaReaderGpu.cpp
@@ -12,6 +12,7 @@
 #include "ExrReaderGpu.h"
 #include "RHICommandList.h"
 #include "ExrSwizzlingShader.h"
+#include "ImgMedia/Public/ImgMediaMipMapInfo.h"
 #include "SceneUtils.h"
 
 
@@ -28,7 +29,6 @@
 #include "ID3D12DynamicRHI.h"
 
 #define READ_IN_CHUNKS 1
-#define EXR_ENABLE_MIPS 1
 
 DECLARE_GPU_STAT_NAMED(ExrImgMediaReaderGpu, TEXT("ExrImgMediaReaderGpu"));
 
@@ -38,24 +38,23 @@ namespace {
 	void DrawScreenPass(
 		FRHICommandListImmediate& RHICmdList,
 		const FIntPoint& OutputResolution,
+		const FIntRect& Viewport,
 		const FScreenPassPipelineState& PipelineState,
 		TSetupFunction SetupFunction)
 	{
-		RHICmdList.SetViewport(0.f, 0.f, 0.f, OutputResolution.X, OutputResolution.Y, 1.0f);
+		RHICmdList.SetViewport(Viewport.Min.X, Viewport.Min.Y, 0.f, Viewport.Max.X, Viewport.Max.Y, 1.0f);
 
 		SetScreenPassPipelineState(RHICmdList, PipelineState);
 
 		// Setting up buffers.
 		SetupFunction(RHICmdList);
 
-		FIntPoint LocalOutputPos(FIntPoint::ZeroValue);
-		FIntPoint LocalOutputSize(OutputResolution);
 		EDrawRectangleFlags DrawRectangleFlags = EDRF_UseTriangleOptimization;
 
 		DrawPostProcessPass(
 			RHICmdList,
-			LocalOutputPos.X, LocalOutputPos.Y, LocalOutputSize.X, LocalOutputSize.Y,
-			0., 0., OutputResolution.X, OutputResolution.Y,
+			0, 0, OutputResolution.X, OutputResolution.Y,
+			Viewport.Min.X, Viewport.Min.Y, Viewport.Width(), Viewport.Height(),
 			OutputResolution,
 			OutputResolution,
 			PipelineState.VertexShader,
@@ -129,6 +128,15 @@ bool FExrImgMediaReaderGpu::ReadFrame(int32 FrameId, int32 MipLevel, const FImgM
 		return false;
 	}
 
+	// Get tile info.
+	int32 NumTilesX = Loader->GetNumTilesX();
+	int32 NumTilesY = Loader->GetNumTilesY();
+	bool bHasTiles = (NumTilesX * NumTilesY) > 1;
+
+	int32 StartTileX = InTileSelection.TopLeftX;
+	int32 StartTileY = InTileSelection.TopLeftY;
+	int32 EndTileX = FMath::Min((int32)InTileSelection.BottomRightX, NumTilesX);
+	int32 EndTileY = FMath::Min((int32)InTileSelection.BottomRightY, NumTilesY);
 	const FString& LargestImagePath = Loader->GetImagePath(FrameId, 0);
 	FRgbaInputFile InputFile(LargestImagePath);
 
@@ -137,9 +145,15 @@ bool FExrImgMediaReaderGpu::ReadFrame(int32 FrameId, int32 MipLevel, const FImgM
 		return false;
 	}
 
-	const FIntPoint& LargestDim = OutFrame->Info.Dim;
+	FIntPoint& FullResolution = OutFrame->Info.Dim;
+	const FIntPoint TileDim = FullResolution;
+	FullResolution.X *= NumTilesX;
+	FullResolution.Y *= NumTilesY;
+	OutFrame->Info.UncompressedSize *= NumTilesX * NumTilesY;
+
+	FIntRect Viewport(StartTileX * TileDim.X, StartTileY * TileDim.Y, EndTileX * TileDim.X, EndTileY * TileDim.Y);
 
-	if (LargestDim.GetMin() <= 0)
+	if (FullResolution.GetMin() <= 0)
 	{
 		return false;
 	}
@@ -152,44 +166,112 @@ bool FExrImgMediaReaderGpu::ReadFrame(int32 FrameId, int32 MipLevel, const FImgM
 	int32 NumMipLevels = Loader->GetNumMipLevels();
 	{
 		// Loop over all mips.
-		FIntPoint Dim = LargestDim;
+		FIntPoint Dim = FullResolution;
+
 		for (int32 CurrentMipLevel = 0; CurrentMipLevel < NumMipLevels; ++CurrentMipLevel)
 		{
 			// Do we want to read in this mip?
 			bool IsThisLevelPresent = (OutFrame->MipMapsPresent & (1 << CurrentMipLevel)) != 0;
-			bool ReadThisMip = (CurrentMipLevel >= MipLevel) &&
-				(IsThisLevelPresent == false);
+			bool ReadThisMip = (CurrentMipLevel >= MipLevel) && (IsThisLevelPresent == false);
 
-#if EXR_ENABLE_MIPS == 0
-			// Just read in mip 0 if we don't already have it.
-			if ((CurrentMipLevel > 0) || (OutFrame->MipMapsPresent != 0))
-			{
-				break;
-			}
-			ReadThisMip = true;
-#endif // EXR_ENABLE_MIPS == 0
+			// Next mip level.
+			Dim = FullResolution / FMath::Pow(2., float(CurrentMipLevel));
 
 			if (ReadThisMip)
 			{
-				// Get for our frame/mip level.
-				const FString& ImagePath = Loader->GetImagePath(FrameId, CurrentMipLevel);
-				EReadResult ReadResult = Fail;
-				
-				const SIZE_T BufferSize = GetBufferSize(Dim, NumChannels);
+				const uint8 PlanePadding = 8;
+				const SIZE_T BufferSize = GetBufferSize(Dim, NumChannels) - Dim.Y * PlanePadding + NumTilesX * NumTilesY * TileDim.Y * PlanePadding;
 				FStructuredBufferPoolItemSharedPtr& BufferData = BufferDataArray[CurrentMipLevel];
 				BufferData = AllocateGpuBufferFromPool(BufferSize);
 				uint16* MipDataPtr = static_cast<uint16*>(BufferData->MappedBuffer);
 
+				// Get for our frame/mip level.
+				FString ImagePath = Loader->GetImagePath(FrameId, CurrentMipLevel);
+				FString BaseImage;
+				if (bHasTiles)
+				{
+					// Remove "_x0_y0.exr" so we can add on the correct name for the tile we want.
+					BaseImage = ImagePath.LeftChop(10);
+				}
+				EReadResult ReadResult = Fail;
+
+				int32 TileWidth = Dim.X / NumTilesX;
+				int32 TileHeight = Dim.Y / NumTilesY;
+
+				FIntPoint TileSize(TileWidth, TileHeight);
+				const SIZE_T DataSizeToRead = GetBufferSize(TileSize, NumChannels);
+
 #if READ_IN_CHUNKS
-				ReadResult = ReadInChunks(MipDataPtr, ImagePath, FrameId, Dim, BufferSize, PixelSize, NumChannels);
+				for (int32 TileY = StartTileY; TileY < EndTileY; TileY++)
+				{
+					for (int32 TileX = StartTileX; TileX < EndTileX; TileX++)
+					{
+						// Get for our frame/mip level.
+						if (bHasTiles)
+						{
+							ImagePath = FString::Printf(TEXT("%s_x%d_y%d.exr"), *BaseImage, TileX, TileY);
+						}
+						FRgbaInputFile InputTileFile(ImagePath);
+						if (InputTileFile.HasInputFile())
+						{
+							// read frame data
+							uint32 CurrentBufferOffset = (NumTilesX * TileY + TileX) * DataSizeToRead / sizeof(uint16);
+							check(CurrentBufferOffset <= BufferSize);
+
+							ReadResult = ReadInChunks(MipDataPtr + CurrentBufferOffset, ImagePath, FrameId, TileSize, DataSizeToRead, PixelSize, NumChannels);
+							
+							if (ReadResult == Fail)
+							{
+								goto ReadFailedFallback;
+							}
+
+							OutFrame->MipMapsPresent |= 1 << CurrentMipLevel;
+						}
+						else
+						{
+							UE_LOG(LogImgMedia, Error, TEXT("Could not load %s"), *ImagePath);
+							return false;
+						}
+					}
+				}
 #else
-
-				bool bResult = FExrReader::GenerateTextureData(MipDataPtr, ImagePath, Dim.X, Dim.Y, PixelSize, NumChannels);
-				ReadResult = bResult ? Success : Fail;
+				for (int32 TileY = StartTileY; TileY < EndTileY; TileY++)
+				{
+					for (int32 TileX = StartTileX; TileX < EndTileX; TileX++)
+					{
+						// Get for our frame/mip level.
+						if (bHasTiles)
+						{
+							ImagePath = FString::Printf(TEXT("%s_x%d_y%d.exr"), *BaseImage, TileX, TileY);
+						}
+						FRgbaInputFile InputTileFile(ImagePath);
+						if (InputTileFile.HasInputFile())
+						{
+							// read frame data
+							uint32 CurrentBufferOffset = (NumTilesX * (TileY)+TileX) * DataSizeToRead / sizeof(uint16);
+							check(CurrentBufferOffset <= BufferSize);
+							bool bResult = FExrReader::GenerateTextureData(MipDataPtr + CurrentBufferOffset, ImagePath, TileSize.X, TileSize.Y, PixelSize, NumChannels);
+							ReadResult = bResult ? Success : Fail;
+
+							if (!bResult)
+							{
+								goto ReadFailedFallback;
+							}
+
+							OutFrame->MipMapsPresent |= 1 << CurrentMipLevel;
+						}
+						else
+						{
+							UE_LOG(LogImgMedia, Error, TEXT("Could not load %s"), *ImagePath);
+							return false;
+						}
+					}
+				}
 #endif
 
 				if (ReadResult == Fail)
 				{
+				ReadFailedFallback:
 					// Check if we have a compressed file.
 					FRgbaInputFile InputFileMip(ImagePath);
 					FImgMediaFrameInfo Info;
@@ -207,20 +289,17 @@ bool FExrImgMediaReaderGpu::ReadFrame(int32 FrameId, int32 MipLevel, const FImgM
 					return FExrImgMediaReader::ReadFrame(FrameId, MipLevel, InTileSelection, OutFrame);
 				}
 			}
-
-			// Next level.
-			Dim /= 2;
 		}
 	}
 
 	OutFrame->Format = NumChannels <= 3 ? EMediaTextureSampleFormat::FloatRGB : EMediaTextureSampleFormat::FloatRGBA;
-	OutFrame->Stride = LargestDim.X * PixelSize;
-	auto RenderThreadSwizzler = [this, BufferDataArray, LargestDim, FrameId, NumChannels, NumMipLevels](FRHICommandListImmediate& RHICmdList, FTexture2DRHIRef RenderTargetTextureRHI)->bool
+	OutFrame->Stride = FullResolution.X * PixelSize;
+	auto RenderThreadSwizzler = [this, BufferDataArray, FullResolution, TileDim, Viewport, NumChannels, NumMipLevels, bHasTiles](FRHICommandListImmediate& RHICmdList, FTexture2DRHIRef RenderTargetTextureRHI)->bool
 	{
 		SCOPED_DRAW_EVENT(RHICmdList, FExrImgMediaReaderGpu_Convert);
 		SCOPED_GPU_STAT(RHICmdList, ExrImgMediaReaderGpu);
 
-		FIntPoint Dim = LargestDim;
+		FIntPoint Dim = FullResolution;
 		for (int32 MipLevel = 0; MipLevel < NumMipLevels; ++MipLevel)
 		{
 			FStructuredBufferPoolItemSharedPtr BufferData = BufferDataArray[MipLevel];
@@ -238,9 +317,10 @@ bool FExrImgMediaReaderGpu::ReadFrame(int32 FrameId, int32 MipLevel, const FImgM
 
 				FExrSwizzlePS::FPermutationDomain PermutationVector;
 				PermutationVector.Set<FExrSwizzlePS::FRgbaSwizzle>(NumChannels);
+				PermutationVector.Set<FExrSwizzlePS::FSwizzleTiles>(bHasTiles);
 				FExrSwizzlePS::FParameters Parameters = FExrSwizzlePS::FParameters();
-				Parameters.TextureWidth = Dim.X;
-				Parameters.TextureHeight = Dim.Y;
+				Parameters.TextureSize = Dim;
+				Parameters.TileSize = TileDim;
 
 				Parameters.UnswizzledBuffer = RHICreateShaderResourceView(BufferData->BufferRef);
 
@@ -250,7 +330,7 @@ bool FExrImgMediaReaderGpu::ReadFrame(int32 FrameId, int32 MipLevel, const FImgM
 				TShaderMapRef<FExrSwizzlePS> SwizzleShaderPS(ShaderMap, PermutationVector);
 
 				FScreenPassPipelineState PipelineState(SwizzleShaderVS, SwizzleShaderPS, TStaticBlendState<>::GetRHI(), TStaticDepthStencilState<false, CF_Always>::GetRHI());
-				DrawScreenPass(RHICmdList, Dim, PipelineState, [&](FRHICommandListImmediate& RHICmdList)
+				DrawScreenPass(RHICmdList, Dim, Viewport, PipelineState, [&](FRHICommandListImmediate& RHICmdList)
 				{
 					SetShaderParameters(RHICmdList, SwizzleShaderPS, SwizzleShaderPS.GetPixelShader(), Parameters);
 				});
@@ -312,7 +392,6 @@ FExrImgMediaReaderGpu::EReadResult FExrImgMediaReaderGpu::ReadInChunks(uint16* B
 	const int32 Remainder = BufferSize % ChunkSize;
 	const int32 NumChunks = (BufferSize - Remainder) / ChunkSize;
 	int32 CurrentBufferPos = 0;
-	
 	FExrReader ChunkReader;
 
 
