ImgMedia: Process images tab writes out attribute data for custom format images.

#rnx

diff --git a/Engine/Plugins/Media/ImgMedia/Source/ImgMedia/Private/ImgMediaModule.cpp b/Engine/Plugins/Media/ImgMedia/Source/ImgMedia/Private/ImgMediaModule.cpp
--- a/Engine/Plugins/Media/ImgMedia/Source/ImgMedia/Private/ImgMediaModule.cpp
+++ b/Engine/Plugins/Media/ImgMedia/Source/ImgMedia/Private/ImgMediaModule.cpp
@@ -15,6 +15,10 @@
 
 DEFINE_LOG_CATEGORY(LogImgMedia);
 
+FLazyName IImgMediaModule::CustomFormatAttributeName(TEXT("EpicGamesCustomFormat"));
+FLazyName IImgMediaModule::CustomFormatTileWidthAttributeName(TEXT("EpicGamesTileWidth"));
+FLazyName IImgMediaModule::CustomFormatTileHeightAttributeName(TEXT("EpicGamesTileHeight"));
+
 TSharedPtr<FImgMediaGlobalCache, ESPMode::ThreadSafe> IImgMediaModule::GlobalCache;
 
 #if USE_IMGMEDIA_DEALLOC_POOL
diff --git a/Engine/Plugins/Media/ImgMedia/Source/ImgMedia/Public/IImgMediaModule.h b/Engine/Plugins/Media/ImgMedia/Source/ImgMedia/Public/IImgMediaModule.h
--- a/Engine/Plugins/Media/ImgMedia/Source/ImgMedia/Public/IImgMediaModule.h
+++ b/Engine/Plugins/Media/ImgMedia/Source/ImgMedia/Public/IImgMediaModule.h
@@ -2,6 +2,7 @@
 
 #pragma once
 
+#include "CoreMinimal.h"
 #include "Templates/SharedPointer.h"
 #include "Modules/ModuleInterface.h"
 
@@ -38,6 +39,13 @@ public:
 	 */
 	static FImgMediaGlobalCache* GetGlobalCache() { return GlobalCache.Get(); }
 
+	/** Name of attribute in the Exr file that marks it as our custom format. */
+	static FLazyName CustomFormatAttributeName;
+	/** Name of attribute in the Exr file for the tile width for our custom format. */
+	static FLazyName CustomFormatTileWidthAttributeName;
+	/** Name of attribute in the Exr file for the tile height for our custom format. */
+	static FLazyName CustomFormatTileHeightAttributeName;
+
 protected:
 
 	/** Holds the global cache. */
diff --git a/Engine/Plugins/Media/ImgMedia/Source/ImgMediaEditor/Private/Widgets/SImgMediaProcessImages.cpp b/Engine/Plugins/Media/ImgMedia/Source/ImgMediaEditor/Private/Widgets/SImgMediaProcessImages.cpp
--- a/Engine/Plugins/Media/ImgMedia/Source/ImgMediaEditor/Private/Widgets/SImgMediaProcessImages.cpp
+++ b/Engine/Plugins/Media/ImgMedia/Source/ImgMediaEditor/Private/Widgets/SImgMediaProcessImages.cpp
@@ -8,6 +8,7 @@
 #include "Framework/Notifications/NotificationManager.h"
 #include "HAL/PlatformFileManager.h"
 #include "IImageWrapperModule.h"
+#include "IImgMediaModule.h"
 #include "ImageWrapperHelper.h"
 #include "ImgMediaEditorModule.h"
 #include "ImgMediaProcessImagesOptions.h"
@@ -279,9 +280,16 @@ void SImgMediaProcessImages::ProcessImageCustom(TSharedPtr<IImageWrapper>& InIma
 	int32 TileHeight = Height / NumTilesY;
 	int32 BytesPerPixel = RawData.Num() / (Width * Height);
 
+	bool bIsTiled = (NumTilesX > 1) || (NumTilesY > 1);
+
 	// Create tiled exr file.
 	FTiledRgbaOutputFile OutFile(FIntPoint(0, 0), FIntPoint(Width - 1, Height - 1),
 		FIntPoint(0, 0), FIntPoint(Width - 1, Height - 1));
+	OutFile.AddIntAttribute(IImgMediaModule::CustomFormatAttributeName.Resolve().ToString(), 1);
+	OutFile.AddIntAttribute(IImgMediaModule::CustomFormatTileWidthAttributeName.Resolve().ToString(),
+		bIsTiled ? TileWidth : 0);
+	OutFile.AddIntAttribute(IImgMediaModule::CustomFormatTileHeightAttributeName.Resolve().ToString(),
+		bIsTiled ? TileHeight : 0);
 	OutFile.CreateOutputFile(InName, TileWidth, TileHeight, 4, false);
 	
 	FIntPoint Stride(1, Width);
