ImgMedia: ExrReader can read out attributes of our custom files.

#rnx

diff --git a/Engine/Plugins/Media/ImgMedia/Source/ImgMedia/Private/Readers/ExrImgMediaReader.cpp b/Engine/Plugins/Media/ImgMedia/Source/ImgMedia/Private/Readers/ExrImgMediaReader.cpp
--- a/Engine/Plugins/Media/ImgMedia/Source/ImgMedia/Private/Readers/ExrImgMediaReader.cpp
+++ b/Engine/Plugins/Media/ImgMedia/Source/ImgMedia/Private/Readers/ExrImgMediaReader.cpp
@@ -11,6 +11,7 @@
 #include "UObject/Class.h"
 #include "UObject/UObjectGlobals.h"
 #include "HardwareInfo.h"
+#include "IImgMediaModule.h"
 #include "ImgMediaLoader.h"
 #include "ImgMediaMipMapInfo.h"
 
@@ -28,6 +29,9 @@ static TAutoConsoleVariable<bool> CVarEnableUncompressedExrGpuReader(
 
 FExrImgMediaReader::FExrImgMediaReader(const TSharedRef<FImgMediaLoader, ESPMode::ThreadSafe>& InLoader)
 	: LoaderPtr(InLoader)
+	, bIsCustomFormat(false)
+	, bIsCustomFormatTiled(false)
+	, CustomFormatTileSize(EForceInit::ForceInitToZero)
 {
 	const UImgMediaSettings* Settings = GetDefault<UImgMediaSettings>();
 	
@@ -230,6 +234,9 @@ void FExrImgMediaReader::CancelFrame(int32 FrameNumber)
 /** Gets reader type (GPU vs CPU) depending on size of EXR and its compression. */
 TSharedPtr<IImgMediaReader, ESPMode::ThreadSafe> FExrImgMediaReader::GetReader(const TSharedRef <FImgMediaLoader, ESPMode::ThreadSafe>& InLoader, FString FirstImageInSequencePath)
 {
+	bool bIsCustomFormat = false;
+	FIntPoint TileSize(EForceInit::ForceInitToZero);
+
 #if defined(PLATFORM_WINDOWS) && PLATFORM_WINDOWS
 	FRgbaInputFile InputFile(FirstImageInSequencePath, 2);
 	if (InputFile.HasInputFile() == false)
@@ -244,17 +251,33 @@ TSharedPtr<IImgMediaReader, ESPMode::ThreadSafe> FExrImgMediaReader::GetReader(c
 	{
 		return MakeShareable(new FExrImgMediaReader(InLoader));
 	}
-	
+
+	// Is this our custom format?
+	int32 CustomFormat = 0;
+	InputFile.GetIntAttribute(IImgMediaModule::CustomFormatAttributeName.Resolve().ToString(), CustomFormat);
+	bIsCustomFormat = CustomFormat > 0;
+	if (bIsCustomFormat)
+	{
+		// Get tile size.
+		InputFile.GetIntAttribute(IImgMediaModule::CustomFormatTileWidthAttributeName.Resolve().ToString(), TileSize.X);
+		InputFile.GetIntAttribute(IImgMediaModule::CustomFormatTileHeightAttributeName.Resolve().ToString(), TileSize.Y);
+	}
+
 	// Check GetCompressionName of OpenExrWrapper for other compression names.
 	if (GDynamicRHI && GDynamicRHI->GetInterfaceType() == ERHIInterfaceType::D3D12
 		&& Info.CompressionName == "Uncompressed" 
 		&& CVarEnableUncompressedExrGpuReader.GetValueOnAnyThread()
 		)
 	{
-		return MakeShared<FExrImgMediaReaderGpu, ESPMode::ThreadSafe>(InLoader);
+		TSharedRef<FExrImgMediaReaderGpu, ESPMode::ThreadSafe> GpuReader = 
+			MakeShared<FExrImgMediaReaderGpu, ESPMode::ThreadSafe>(InLoader);
+		GpuReader->SetCustomFormatInfo(bIsCustomFormat, TileSize);
+		return GpuReader;
 	}
 #endif
-	return MakeShareable(new FExrImgMediaReader(InLoader));
+	FExrImgMediaReader* Reader = new FExrImgMediaReader(InLoader);
+	Reader->SetCustomFormatInfo(bIsCustomFormat, TileSize);
+	return MakeShareable(Reader);
 }
 
 /* FExrImgMediaReader implementation
@@ -278,6 +301,13 @@ bool FExrImgMediaReader::GetInfo(FRgbaInputFile& InputFile, FImgMediaFrameInfo&
 	return (OutInfo.UncompressedSize > 0) && (OutInfo.Dim.GetMin() > 0);
 }
 
+void FExrImgMediaReader::SetCustomFormatInfo(bool bInIsCustomFormat, const FIntPoint& InTileSize)
+{
+	bIsCustomFormat = bInIsCustomFormat;
+	CustomFormatTileSize = InTileSize;
+	bIsCustomFormatTiled = InTileSize.X != 0;
+}
+
 SIZE_T FExrImgMediaReader::GetMipBufferTotalSize(FIntPoint Dim)
 {
 	SIZE_T Size = ((Dim.X * Dim.Y * 4) / 3) * sizeof(uint16) * 4;
diff --git a/Engine/Plugins/Media/ImgMedia/Source/ImgMedia/Private/Readers/ExrImgMediaReader.h b/Engine/Plugins/Media/ImgMedia/Source/ImgMedia/Private/Readers/ExrImgMediaReader.h
--- a/Engine/Plugins/Media/ImgMedia/Source/ImgMedia/Private/Readers/ExrImgMediaReader.h
+++ b/Engine/Plugins/Media/ImgMedia/Source/ImgMedia/Private/Readers/ExrImgMediaReader.h
@@ -35,6 +35,11 @@ public:
 	/** Gets reader type (GPU vs CPU) depending on size of EXR and its compression. */
 	static TSharedPtr<IImgMediaReader, ESPMode::ThreadSafe> GetReader(const TSharedRef<FImgMediaLoader, ESPMode::ThreadSafe>& InLoader, FString FirstImageInSequencePath);
 
+	/** Query if our images are in our custom format. */
+	bool IsCustomFormat() const { return bIsCustomFormat; }
+	/** Gets the tile size of our custom format. */
+	FIntPoint GetCustomFormatTileSize() { return CustomFormatTileSize; }
+
 protected:
 
 	/**
@@ -46,6 +51,14 @@ protected:
 	 */
 	static bool GetInfo(FRgbaInputFile& InputFile, FImgMediaFrameInfo& OutInfo);
 
+	/**
+	 * Sets parameters of our custom format images.
+	 *
+	 * @param bInIsCustomFormat		True if we are using custom format images.
+	 * @param TileSize				Size of our tiles. If (0, 0) then we are not using tiles.
+	 */
+	void SetCustomFormatInfo(bool bInIsCustomFormat, const FIntPoint& InTileSize);
+
 	/**
 	 * Gets the total size needed for all mips.
 	 *
@@ -59,6 +72,12 @@ protected:
 	FCriticalSection CanceledFramesCriticalSection;
 	TMap<int32, FRgbaInputFile*> PendingFrames;
 	TWeakPtr<FImgMediaLoader, ESPMode::ThreadSafe> LoaderPtr;
+	/** True if we are using our custom format. */
+	bool bIsCustomFormat;
+	/** True if our custom format images are tiled. */
+	bool bIsCustomFormatTiled;
+	/** Tile size of our custom format images. */
+	FIntPoint CustomFormatTileSize;
 };
 
 
