ImgMedia: Can use the EXR general interface to write out files in OpenExrWrapper.

#rnx

diff --git a/Engine/Plugins/Media/ImgMedia/Source/OpenExrWrapper/Private/OpenExrWrapper.cpp b/Engine/Plugins/Media/ImgMedia/Source/OpenExrWrapper/Private/OpenExrWrapper.cpp
--- a/Engine/Plugins/Media/ImgMedia/Source/OpenExrWrapper/Private/OpenExrWrapper.cpp
+++ b/Engine/Plugins/Media/ImgMedia/Source/OpenExrWrapper/Private/OpenExrWrapper.cpp
@@ -9,11 +9,14 @@
 PRAGMA_DEFAULT_VISIBILITY_START
 THIRD_PARTY_INCLUDES_START
 	#include "Imath/ImathBox.h"
+	#include "OpenEXR/ImfChannelList.h"
 	#include "OpenEXR/ImfCompressionAttribute.h"
 	#include "OpenEXR/ImfHeader.h"
 	#include "OpenEXR/ImfIntAttribute.h"
 	#include "OpenEXR/ImfRgbaFile.h"
 	#include "OpenEXR/ImfStandardAttributes.h"
+	#include "OpenEXR/ImfTiledInputFile.h"
+	#include "OpenEXR/ImfTiledOutputFile.h"
 	#include "OpenEXR/ImfTiledRgbaFile.h"
 THIRD_PARTY_INCLUDES_END
 PRAGMA_DEFAULT_VISIBILITY_END
@@ -234,7 +237,7 @@ bool FRgbaInputFile::GetIntAttribute(const FString& Name, int32& Value)
 	return bIsAttributeFound;
 }
 
-FTiledRgbaOutputFile::FTiledRgbaOutputFile(
+FBaseOutputFile::FBaseOutputFile(
 	const FIntPoint& DisplayWindowMin,
 	const FIntPoint& DisplayWindowMax,
 	const FIntPoint& DataWindowMin,
@@ -251,7 +254,7 @@ FTiledRgbaOutputFile::FTiledRgbaOutputFile(
 		Imf::NO_COMPRESSION);
 }
 
-FTiledRgbaOutputFile::~FTiledRgbaOutputFile()
+FBaseOutputFile::~FBaseOutputFile()
 {
 	if (Header != nullptr)
 	{
@@ -263,12 +266,12 @@ FTiledRgbaOutputFile::~FTiledRgbaOutputFile()
 	}
 }
 
-void FTiledRgbaOutputFile::AddIntAttribute(const FString& Name, int32 Value)
+void FBaseOutputFile::AddIntAttribute(const FString& Name, int32 Value)
 {
 	// Make sure we don't have an output file yet.
 	if (OutputFile == nullptr)
 	{
-		((Imf::Header*)Header)->insert(std::string(TCHAR_TO_ANSI(*Name)), Imf::IntAttribute(Value));
+		((Imf::Header*)Header)->insert(std::string(StringCast<ANSICHAR>(*Name).Get()), Imf::IntAttribute(Value));
 	}
 	else
 	{
@@ -277,6 +280,16 @@ void FTiledRgbaOutputFile::AddIntAttribute(const FString& Name, int32 Value)
 	}
 }
 
+
+FTiledRgbaOutputFile::FTiledRgbaOutputFile(
+	const FIntPoint& DisplayWindowMin,
+	const FIntPoint& DisplayWindowMax,
+	const FIntPoint& DataWindowMin,
+	const FIntPoint& DataWindowMax)
+	: FBaseOutputFile(DisplayWindowMin, DisplayWindowMax, DataWindowMin, DataWindowMax)
+{
+}
+
 void FTiledRgbaOutputFile::CreateOutputFile(const FString& FilePath, 
 	int32 TileWidth, int32 TileHeight, int32 NumChannels, bool bIsMipsEnabled)
 {
@@ -308,7 +321,7 @@ void FTiledRgbaOutputFile::CreateOutputFile(const FString& FilePath,
 			}
 
 			// Create output file.
-			OutputFile = new Imf::TiledRgbaOutputFile(TCHAR_TO_ANSI(*FilePath),
+			OutputFile = new Imf::TiledRgbaOutputFile(StringCast<ANSICHAR>(*FilePath).Get(),
 				*((Imf::Header*)Header),
 				Channels,
 				TileWidth, TileHeight,
@@ -376,4 +389,96 @@ void FTiledRgbaOutputFile::WriteTile(int32 TileX, int32 TileY, int32 MipLevel)
 	}
 }
 
+FTiledOutputFile::FTiledOutputFile(
+	const FIntPoint& DisplayWindowMin,
+	const FIntPoint& DisplayWindowMax,
+	const FIntPoint& DataWindowMin,
+	const FIntPoint& DataWindowMax)
+	: FBaseOutputFile(DisplayWindowMin, DisplayWindowMax, DataWindowMin, DataWindowMax)
+{
+	FrameBuffer = new Imf::FrameBuffer;
+}
+
+FTiledOutputFile::~FTiledOutputFile()
+{
+	if (FrameBuffer != nullptr)
+	{
+		delete (Imf::FrameBuffer*)FrameBuffer;
+	}
+}
+
+void FTiledOutputFile::AddChannel(const FString& Name)
+{
+	((Imf::Header*)Header)->channels().insert(StringCast<ANSICHAR>(*Name).Get(), Imf::Channel(Imf::HALF));
+}
+
+void FTiledOutputFile::CreateOutputFile(const FString& FilePath,
+	int32 TileWidth, int32 TileHeight, bool bIsMipsEnabled)
+{
+	if (OutputFile == nullptr)
+	{
+		try
+		{
+			((Imf::Header*)Header)->setTileDescription(Imf::TileDescription(TileWidth, TileHeight,
+				bIsMipsEnabled ? Imf::MIPMAP_LEVELS : Imf::ONE_LEVEL));
+
+
+			// Create output file.
+			OutputFile = new Imf::TiledOutputFile(StringCast<ANSICHAR>(*FilePath).Get(),
+				*((Imf::Header*)Header));
+		}
+		catch (std::exception const& Exception)
+		{
+			UE_LOG(LogOpenEXRWrapper, Error, TEXT("Cannot write EXR file: %s (%s)"),
+				*FilePath, StringCast<TCHAR>(Exception.what()).Get());
+		}
+	}
+	else
+	{
+		UE_LOG(LogOpenEXRWrapper, Error,
+			TEXT("Cannot create output file as it has already been created."));
+	}
+}
+
+void FTiledOutputFile::AddFrameBufferChannel(const FString& Name, void* Base,
+	const FIntPoint& Stride)
+{
+	((Imf::FrameBuffer*)FrameBuffer)->insert(StringCast<ANSICHAR>(*Name).Get(),
+		Imf::Slice(Imf::HALF, (char*)Base, Stride.X, Stride.Y));
+}
+
+void FTiledOutputFile::SetFrameBuffer()
+{
+	if (OutputFile != nullptr)
+	{
+		((Imf::TiledOutputFile*)OutputFile)->setFrameBuffer(*((Imf::FrameBuffer*)FrameBuffer));
+	}
+	else
+	{
+		UE_LOG(LogOpenEXRWrapper, Error,
+			TEXT("Cannot set frame buffer as there is no output file."));
+	}
+}
+
+void FTiledOutputFile::WriteTile(int32 TileX, int32 TileY, int32 MipLevel)
+{
+	if (OutputFile != nullptr)
+	{
+		try
+		{
+			((Imf::TiledOutputFile*)OutputFile)->writeTile(TileX, TileY, MipLevel);
+		}
+		catch (std::exception const& Exception)
+		{
+			UE_LOG(LogOpenEXRWrapper, Error, TEXT("Cannot write EXR file: %s"),
+				StringCast<TCHAR>(Exception.what()).Get());
+		}
+	}
+	else
+	{
+		UE_LOG(LogOpenEXRWrapper, Error,
+			TEXT("WriteTile failed: CreateOutputFile has not been called yet."));
+	}
+}
+
 IMPLEMENT_MODULE(FDefaultModuleImpl, OpenExrWrapper);
diff --git a/Engine/Plugins/Media/ImgMedia/Source/OpenExrWrapper/Public/OpenExrWrapper.h b/Engine/Plugins/Media/ImgMedia/Source/OpenExrWrapper/Public/OpenExrWrapper.h
--- a/Engine/Plugins/Media/ImgMedia/Source/OpenExrWrapper/Public/OpenExrWrapper.h
+++ b/Engine/Plugins/Media/ImgMedia/Source/OpenExrWrapper/Public/OpenExrWrapper.h
@@ -54,24 +54,20 @@ private:
 };
 
 /**
- * Use this to write out tiled EXR images.
- * 
- * Add any attributes after construction.
- * Then you can call CreateOutputFile.
- * After that, you can write out data to the file.
+ * Base class for our classes that output EXR files.
  */
-class OPENEXRWRAPPER_API FTiledRgbaOutputFile
+class OPENEXRWRAPPER_API FBaseOutputFile
 {
 public:
 	/**
 	 * Constructor.
-	 * 
+	 *
 	 * @param DisplayWindowMin		Normally (0, 0).
 	 * @param DisplayWindowMax		Normally (width - 1, height - 1).
 	 * @param DataWindowMin			Normally (0, 0).
 	 * @param DataWindowMax			Normally (width - 1, height - 1).
 	 */
-	FTiledRgbaOutputFile(
+	FBaseOutputFile(
 		const FIntPoint& DisplayWindowMin,
 		const FIntPoint& DisplayWindowMax,
 		const FIntPoint& DataWindowMin,
@@ -80,17 +76,48 @@ public:
 	/**
 	 * Destructor.
 	 */
-	~FTiledRgbaOutputFile();
+	virtual ~FBaseOutputFile();
 
 	/**
 	 * Call this to add an attribute to the EXR file.
 	 * This MUST be called before CreateOutputFile.
-	 * 
+	 *
 	 * @param Name		Name for this attribute.
 	 * @param Value		Value for this attribute.
 	 */
 	void AddIntAttribute(const FString& Name, int32 Value);
 
+protected:
+	/** Stores the EXR header object. */
+	void* Header;
+	/** Stores the EXR object. */
+	void* OutputFile;
+};
+
+/**
+ * Use this to write out tiled EXR images.
+ * 
+ * Add any attributes after construction.
+ * Then you can call CreateOutputFile.
+ * After that, you can write out data to the file.
+ */
+class OPENEXRWRAPPER_API FTiledRgbaOutputFile : public FBaseOutputFile
+{
+public:
+	/**
+	 * Constructor.
+	 * 
+	 * @param DisplayWindowMin		Normally (0, 0).
+	 * @param DisplayWindowMax		Normally (width - 1, height - 1).
+	 * @param DataWindowMin			Normally (0, 0).
+	 * @param DataWindowMax			Normally (width - 1, height - 1).
+	 */
+	FTiledRgbaOutputFile(
+		const FIntPoint& DisplayWindowMin,
+		const FIntPoint& DisplayWindowMax,
+		const FIntPoint& DataWindowMin,
+		const FIntPoint& DataWindowMax);
+
 	/**
 	 * Call this after adding any attributes BUT before doing anything else.
 	 *
@@ -116,6 +143,81 @@ public:
 	 */
 	void SetFrameBuffer(void* Buffer, const FIntPoint& Stride);
 
+	/**
+	 * Call this to write data to a specific tile.
+	 *
+	 * @param TileX			X coordinate of tile.
+	 * @param TileY			Y coordinate of tile.
+	 * @param MipLevel		Mipmap level of tile.
+	 */
+	void WriteTile(int32 TileX, int32 TileY, int32 MipLevel);
+};
+
+/**
+ * Use this to write out tiled EXR images using the general interface.
+ *
+ * Add any attributes and AddChannel after construction.
+ * Then you can call CreateOutputFile.
+ * Then you can call AddFrameBufferChannel for each channel, and then SetFrameBuffer.
+ * After that, you can write out data to the file.
+ */
+class OPENEXRWRAPPER_API FTiledOutputFile : public FBaseOutputFile
+{
+public:
+	/**
+	 * Constructor.
+	 *
+	 * @param DisplayWindowMin		Normally (0, 0).
+	 * @param DisplayWindowMax		Normally (width - 1, height - 1).
+	 * @param DataWindowMin			Normally (0, 0).
+	 * @param DataWindowMax			Normally (width - 1, height - 1).
+	 */
+	FTiledOutputFile(
+		const FIntPoint& DisplayWindowMin,
+		const FIntPoint& DisplayWindowMax,
+		const FIntPoint& DataWindowMin,
+		const FIntPoint& DataWindowMax);
+
+	/**
+	 * Destructor.
+	 */
+	virtual ~FTiledOutputFile();
+
+	/**
+	 * Call this before CreateOutputFile for each channel in the image.
+	 *
+	 * @param Name		Name of channel (e.g. R, G, B, or A). 
+	 */
+	void AddChannel(const FString& Name);
+
+	/**
+	 * Call this after adding any attributes or channels BUT before doing anything else.
+	 *
+	 * @param FilePath			Filename to save to.
+	 * @param TileWidth			Width of a tile.
+	 * @param TileHeight		Height of a tile.
+	 * @param bIsMipsEnabled	True to enable mip mapping.
+	 */
+	void CreateOutputFile(const FString& FilePath,
+		int32 TileWidth, int32 TileHeight, bool bIsMipsEnabled);
+
+	/**
+	 * Call this before SetFrameBuffer for each channel in the frame buffer.
+	 *
+	 * @param Name			Name of channel (e.g. R, G, B, or A).
+	 * @param Base			Address of the start of the data.
+	 * @param Stride		A pixels location is calculated by Base + x * StrideX + y * StrideY.
+	 */
+	void AddFrameBufferChannel(const FString& Name, void* Base,
+		const FIntPoint& Stride);
+	/**
+	 * Call this prior to WriteTile to set where the data is coming from.
+	 *
+	 * @param Buffer		Source of data.
+	 * @param Stride		A pixels location is calculated by Buffer + x * StrideX + y * StrideY.
+	 */
+	void SetFrameBuffer();
+
 	/**
 	 * Call this to write data to a specific tile.
 	 *
@@ -126,8 +228,8 @@ public:
 	void WriteTile(int32 TileX, int32 TileY, int32 MipLevel);
 
 private:
-	/** Stores the EXR header object. */
-	void* Header;
-	/** Stores the EXR object. */
-	void* OutputFile;
+
+	/** Stores the EXR frame buffer object. */
+	void* FrameBuffer;
 };
+
