ImgMedia: Process images tab writes out unswizzled pixel data for the custom format.

#rnx

diff --git a/Engine/Plugins/Media/ImgMedia/Source/ImgMediaEditor/Private/Widgets/SImgMediaProcessImages.cpp b/Engine/Plugins/Media/ImgMedia/Source/ImgMediaEditor/Private/Widgets/SImgMediaProcessImages.cpp
--- a/Engine/Plugins/Media/ImgMedia/Source/ImgMediaEditor/Private/Widgets/SImgMediaProcessImages.cpp
+++ b/Engine/Plugins/Media/ImgMedia/Source/ImgMediaEditor/Private/Widgets/SImgMediaProcessImages.cpp
@@ -282,17 +282,40 @@ void SImgMediaProcessImages::ProcessImageCustom(TSharedPtr<IImageWrapper>& InIma
 
 	bool bIsTiled = (NumTilesX > 1) || (NumTilesY > 1);
 
+	// Names for our channels.
+	const FString RChannelName = FString(TEXT("R"));
+	const FString GChannelName = FString(TEXT("G"));
+	const FString BChannelName = FString(TEXT("B"));
+	const FString AChannelName = FString(TEXT("A"));
+
+	int32 NumChannels = 4;
+
+	FIntPoint Stride(2, Width * BytesPerPixel);
+
 	// Create tiled exr file.
-	FTiledRgbaOutputFile OutFile(FIntPoint(0, 0), FIntPoint(Width - 1, Height - 1),
+	FTiledOutputFile OutFile(FIntPoint(0, 0), FIntPoint(Width - 1, Height - 1),
 		FIntPoint(0, 0), FIntPoint(Width - 1, Height - 1));
+
+	// Add attributes.
 	OutFile.AddIntAttribute(IImgMediaModule::CustomFormatAttributeName.Resolve().ToString(), 1);
 	OutFile.AddIntAttribute(IImgMediaModule::CustomFormatTileWidthAttributeName.Resolve().ToString(),
 		bIsTiled ? TileWidth : 0);
 	OutFile.AddIntAttribute(IImgMediaModule::CustomFormatTileHeightAttributeName.Resolve().ToString(),
 		bIsTiled ? TileHeight : 0);
-	OutFile.CreateOutputFile(InName, TileWidth, TileHeight, 4, false);
-	
-	FIntPoint Stride(1, Width);
+
+	// Add channels.
+	OutFile.AddChannel(AChannelName);
+	OutFile.AddChannel(BChannelName);
+	OutFile.AddChannel(GChannelName);
+	OutFile.AddChannel(RChannelName);
+
+	// Create output.
+	OutFile.CreateOutputFile(InName, TileWidth, TileHeight, false);
+	OutFile.AddFrameBufferChannel(AChannelName, RawData.GetData(), Stride);
+	OutFile.AddFrameBufferChannel(BChannelName, RawData.GetData() + Width * 2, Stride);
+	OutFile.AddFrameBufferChannel(GChannelName, RawData.GetData() + Width * 4, Stride);
+	OutFile.AddFrameBufferChannel(RChannelName, RawData.GetData() + Width * 6, Stride);
+	OutFile.SetFrameBuffer();
 	
 	// Loop over y tiles.
 	for (int TileY = 0; TileY < NumTilesY; ++TileY)
@@ -300,10 +323,12 @@ void SImgMediaProcessImages::ProcessImageCustom(TSharedPtr<IImageWrapper>& InIma
 		// Loop over x tiles.
 		for (int TileX = 0; TileX < NumTilesX; ++TileX)
 		{
-			OutFile.SetFrameBuffer(RawData.GetData(), Stride);
-			OutFile.WriteTile(TileX, TileY, 0);
+			
 		}
 	}
+	
+	OutFile.WriteTile(0, 0, 0);
+
 #else // IMGMEDIAEDITOR_EXR_SUPPORTED_PLATFORM
 	UE_LOG(LogImgMediaEditor, Error, TEXT("EXR not supported on this platform."));
 #endif // IMGMEDIAEDITOR_EXR_SUPPORTED_PLATFORM
