Image Media, ExrImgMediaReader: Added ability to read tile information from EXR files.

diff --git a/Engine/Plugins/Media/ImgMedia/Source/ImgMedia/Private/Readers/ExrImgMediaReader.cpp b/Engine/Plugins/Media/ImgMedia/Source/ImgMedia/Private/Readers/ExrImgMediaReader.cpp
--- a/Engine/Plugins/Media/ImgMedia/Source/ImgMedia/Private/Readers/ExrImgMediaReader.cpp
+++ b/Engine/Plugins/Media/ImgMedia/Source/ImgMedia/Private/Readers/ExrImgMediaReader.cpp
@@ -297,6 +297,11 @@ bool FExrImgMediaReader::GetInfo(FRgbaInputFile& InputFile, FImgMediaFrameInfo&
 	OutInfo.Srgb = false;
 	OutInfo.UncompressedSize = InputFile.GetUncompressedSize();
 	OutInfo.NumChannels = InputFile.GetNumChannels();
+	OutInfo.bHasTiles = InputFile.GetTileSize(OutInfo.TileDimensions);
+	if (OutInfo.bHasTiles)
+	{
+		OutInfo.NumTiles = FIntPoint(OutInfo.Dim.X / OutInfo.TileDimensions.X, OutInfo.Dim.Y / OutInfo.TileDimensions.Y);
+	}
 
 	return (OutInfo.UncompressedSize > 0) && (OutInfo.Dim.GetMin() > 0);
 }
diff --git a/Engine/Plugins/Media/ImgMedia/Source/ImgMedia/Private/Readers/IImgMediaReader.h b/Engine/Plugins/Media/ImgMedia/Source/ImgMedia/Private/Readers/IImgMediaReader.h
--- a/Engine/Plugins/Media/ImgMedia/Source/ImgMedia/Private/Readers/IImgMediaReader.h
+++ b/Engine/Plugins/Media/ImgMedia/Source/ImgMedia/Private/Readers/IImgMediaReader.h
@@ -37,6 +37,15 @@ struct FImgMediaFrameInfo
 
 	/** Number of channels (RGB - 3 or RGBA - 4). */
 	SIZE_T NumChannels;
+
+	/** Does this frame have tiles?. */
+	bool bHasTiles;
+
+	/** Number of tiles in X and Y direction. */
+	FIntPoint NumTiles;
+
+	/** Tile dimensions. */
+	FIntPoint TileDimensions;
 };
 
 
diff --git a/Engine/Plugins/Media/ImgMedia/Source/OpenExrWrapper/Private/OpenExrWrapper.cpp b/Engine/Plugins/Media/ImgMedia/Source/OpenExrWrapper/Private/OpenExrWrapper.cpp
--- a/Engine/Plugins/Media/ImgMedia/Source/OpenExrWrapper/Private/OpenExrWrapper.cpp
+++ b/Engine/Plugins/Media/ImgMedia/Source/OpenExrWrapper/Private/OpenExrWrapper.cpp
@@ -13,6 +13,7 @@ THIRD_PARTY_INCLUDES_START
 	#include "OpenEXR/ImfCompressionAttribute.h"
 	#include "OpenEXR/ImfHeader.h"
 	#include "OpenEXR/ImfIntAttribute.h"
+	#include "OpenEXR/ImfTileDescriptionAttribute.h"
 	#include "OpenEXR/ImfRgbaFile.h"
 	#include "OpenEXR/ImfStandardAttributes.h"
 	#include "OpenEXR/ImfTiledInputFile.h"
@@ -170,6 +171,18 @@ int32 FRgbaInputFile::GetNumChannels() const
 	return NumChannels;
 }
 
+bool FRgbaInputFile::GetTileSize(FIntPoint& OutTileSize) const
+{
+	const Imf::TileDescriptionAttribute* TileDescAttr = ((Imf::RgbaInputFile*)InputFile)->header().findTypedAttribute<Imf::TileDescriptionAttribute>("tiles");
+	if (TileDescAttr)
+	{
+		Imf::TileDescription TileDesc = TileDescAttr->value();
+		OutTileSize = FIntPoint(TileDesc.xSize, TileDesc.ySize);
+	}
+	return TileDescAttr != nullptr;
+	
+}
+
 int32 FRgbaInputFile::GetUncompressedSize() const
 {
 	const int32 NumChannels = GetNumChannels();
diff --git a/Engine/Plugins/Media/ImgMedia/Source/OpenExrWrapper/Public/OpenExrWrapper.h b/Engine/Plugins/Media/ImgMedia/Source/OpenExrWrapper/Public/OpenExrWrapper.h
--- a/Engine/Plugins/Media/ImgMedia/Source/OpenExrWrapper/Public/OpenExrWrapper.h
+++ b/Engine/Plugins/Media/ImgMedia/Source/OpenExrWrapper/Public/OpenExrWrapper.h
@@ -33,6 +33,9 @@ public:
 	int32 GetUncompressedSize() const;
 	int32 GetNumChannels() const;
 
+	// Gets tile dimensions. Returns false if image has no tiles.
+	bool GetTileSize(FIntPoint& OutTileSize) const;
+
 	bool IsComplete() const;
 	bool HasInputFile() const;
 	void ReadPixels(int32 StartY, int32 EndY);
