ImgMedia ExrImgMediaReaderGpu: Reader now reads and swizzles tiles from native EXR file format.

diff --git a/Engine/Plugins/Media/ImgMedia/Shaders/Private/ExrSwizzler.usf b/Engine/Plugins/Media/ImgMedia/Shaders/Private/ExrSwizzler.usf
--- a/Engine/Plugins/Media/ImgMedia/Shaders/Private/ExrSwizzler.usf
+++ b/Engine/Plugins/Media/ImgMedia/Shaders/Private/ExrSwizzler.usf
@@ -82,14 +82,20 @@ void GetStartScanlinePosition(in int2 PixelPos, out float2 SampleSize, out int S
 	int MaxTilesX = TextureSize.x / TileSize.x;
 	int MaxTilesY = TextureSize.y / TileSize.y;
 
-	// Current scanline position.
-	int TileBufferStride = TileSize.x * TileSize.y * Channels + TileSize.y * 4;
-	int PreviousTilesOffset = (MaxTilesX * (TileY)+TileX) * TileBufferStride;
-	StartPosition = PreviousTilesOffset + yCoord * TileSize.x * Channels + xCoord + (yCoord + 1) * 4;
+	// Padding in bytes
+	const int TilePadding = 10;
+
+	// Current tile position.
+	int TileBufferStride = TileSize.x * TileSize.y * Channels;
+	int PreviousTilesOffset = (MaxTilesX * TileY + TileX) * TileBufferStride;
+	StartPosition = PreviousTilesOffset + yCoord * TileSize.x * Channels + xCoord + (TileY * MaxTilesX + TileX + 1) * TilePadding;
 #else
 	SampleSize.x = TextureSize.x;
 	SampleSize.y = TextureSize.y;
 
+	// Padding in bytes
+	const int TilePadding = 4;
+
 	// Current scanline position.
 	StartPosition = PixelPos.y * TextureSize.x * Channels + PixelPos.x + (PixelPos.y + 1) * 4;
 #endif
diff --git a/Engine/Plugins/Media/ImgMedia/Source/ExrReaderGpu/Private/ExrReaderGpu.cpp b/Engine/Plugins/Media/ImgMedia/Source/ExrReaderGpu/Private/ExrReaderGpu.cpp
--- a/Engine/Plugins/Media/ImgMedia/Source/ExrReaderGpu/Private/ExrReaderGpu.cpp
+++ b/Engine/Plugins/Media/ImgMedia/Source/ExrReaderGpu/Private/ExrReaderGpu.cpp
@@ -216,7 +216,7 @@ bool FExrReader::ReadLineOffsets(FILE* FileHandle, ELineOrder LineOrder, TArray<
 	return true;
 }
 
-bool FExrReader::GenerateTextureData(uint16* Buffer, FString FilePath, int32 TextureWidth, int32 TextureHeight, int32 PixelSize, int32 NumChannels)
+bool FExrReader::GenerateTextureData(uint16* Buffer, int32 BufferSize, FString FilePath, int32 NumberOfScanlines, int32 NumChannels)
 {
 	check(Buffer != nullptr);
 
@@ -237,7 +237,7 @@ bool FExrReader::GenerateTextureData(uint16* Buffer, FString FilePath, int32 Tex
 	{
 		ReadMagicNumberAndVersionField(FileHandle);
 		ReadHeaderData(FileHandle);
-		LineOffsets.SetNum(TextureHeight);
+		LineOffsets.SetNum(NumberOfScanlines);
 
 		// At the moment we support only increasing Y.
 		ELineOrder LineOrder = INCREASING_Y;
@@ -245,7 +245,7 @@ bool FExrReader::GenerateTextureData(uint16* Buffer, FString FilePath, int32 Tex
 	}
 
 	fseek(FileHandle, LineOffsets[0], SEEK_SET);
-	fread(Buffer, TextureWidth * PixelSize + PLANAR_RGB_SCANLINE_PADDING, TextureHeight, FileHandle);
+	fread(Buffer, BufferSize, 1 /*NumOfElements*/, FileHandle);
 
 	fclose(FileHandle);
 	return true;
diff --git a/Engine/Plugins/Media/ImgMedia/Source/ExrReaderGpu/Public/ExrReaderGpu.h b/Engine/Plugins/Media/ImgMedia/Source/ExrReaderGpu/Public/ExrReaderGpu.h
--- a/Engine/Plugins/Media/ImgMedia/Source/ExrReaderGpu/Public/ExrReaderGpu.h
+++ b/Engine/Plugins/Media/ImgMedia/Source/ExrReaderGpu/Public/ExrReaderGpu.h
@@ -51,7 +51,7 @@ private:
 public:
 
 	/** Discards the header and fills the buffer with plain pixel data. */
-	static bool GenerateTextureData(uint16* Buffer, FString FilePath, int32 TextureWidth, int32 TextureHeight, int32 PixelSize, int32 NumChannels);
+	static bool GenerateTextureData(uint16* Buffer, int32 BufferSize, FString FilePath, int32 NumberOfScanlines, int32 NumChannels);
 
 	/** 
 	* This function is used to open file and keep a handle to it so the file can be read in chunks.
diff --git a/Engine/Plugins/Media/ImgMedia/Source/ImgMedia/Private/Readers/ExrImgMediaReaderGpu.cpp b/Engine/Plugins/Media/ImgMedia/Source/ImgMedia/Private/Readers/ExrImgMediaReaderGpu.cpp
--- a/Engine/Plugins/Media/ImgMedia/Source/ImgMedia/Private/Readers/ExrImgMediaReaderGpu.cpp
+++ b/Engine/Plugins/Media/ImgMedia/Source/ImgMedia/Private/Readers/ExrImgMediaReaderGpu.cpp
@@ -33,6 +33,13 @@
 DECLARE_GPU_STAT_NAMED(ExrImgMediaReaderGpu, TEXT("ExrImgMediaReaderGpu"));
 
 namespace {
+
+	/* Padding in bits that is added at the beginning of each scanline */
+	const uint16 ExrRowPadding = 8;
+
+	/* Padding in bits that is added at the beginning of each tile */
+	const uint16 ExrTilePadding = 20;
+
 	/** This function is similar to DrawScreenPass in OpenColorIODisplayExtension.cpp except it is catered for Viewless texture rendering. */
 	template<typename TSetupFunction>
 	void DrawScreenPass(
@@ -128,15 +135,7 @@ bool FExrImgMediaReaderGpu::ReadFrame(int32 FrameId, int32 MipLevel, const FImgM
 		return false;
 	}
 
-	// Get tile info.
-	int32 NumTilesX = Loader->GetNumTilesX();
-	int32 NumTilesY = Loader->GetNumTilesY();
-	bool bHasTiles = (NumTilesX * NumTilesY) > 1;
 
-	int32 StartTileX = InTileSelection.TopLeftX;
-	int32 StartTileY = InTileSelection.TopLeftY;
-	int32 EndTileX = FMath::Min((int32)InTileSelection.BottomRightX, NumTilesX);
-	int32 EndTileY = FMath::Min((int32)InTileSelection.BottomRightY, NumTilesY);
 	const FString& LargestImagePath = Loader->GetImagePath(FrameId, 0);
 	FRgbaInputFile InputFile(LargestImagePath);
 
@@ -145,28 +144,39 @@ bool FExrImgMediaReaderGpu::ReadFrame(int32 FrameId, int32 MipLevel, const FImgM
 		return false;
 	}
 
-	FIntPoint& FullResolution = OutFrame->Info.Dim;
-	const FIntPoint TileDim = FullResolution;
-	FullResolution.X *= NumTilesX;
-	FullResolution.Y *= NumTilesY;
-	OutFrame->Info.UncompressedSize *= NumTilesX * NumTilesY;
+	// Get tile info.
+	bool bHasTiles = OutFrame->Info.bHasTiles;
+	int32 StartTileX = InTileSelection.TopLeftX;
+	int32 StartTileY = InTileSelection.TopLeftY;
+	int32 EndTileX = FMath::Min((int32)InTileSelection.BottomRightX, OutFrame->Info.NumTiles.X);
+	int32 EndTileY = FMath::Min((int32)InTileSelection.BottomRightY, OutFrame->Info.NumTiles.Y);
+
+	const FIntPoint& FullResolution = OutFrame->Info.Dim;
 
-	FIntRect Viewport(StartTileX * TileDim.X, StartTileY * TileDim.Y, EndTileX * TileDim.X, EndTileY * TileDim.Y);
+	const FIntPoint TileDim = OutFrame->Info.TileDimensions;
 
 	if (FullResolution.GetMin() <= 0)
 	{
 		return false;
 	}
 
+	FIntRect TopMipViewport = FIntRect(FIntPoint(0, 0), FullResolution);
+
+	if (bHasTiles)
+	{
+		TopMipViewport = FIntRect(StartTileX * TileDim.X, StartTileY * TileDim.Y, EndTileX * TileDim.X, EndTileY * TileDim.Y);
+	}
+
 	const int32 NumChannels = OutFrame->Info.NumChannels;
 	const int32 PixelSize = sizeof(uint16) * NumChannels;
 
-	FStructuredBufferPoolItemSharedPtr BufferDataArray[FImgMediaLoader::MAX_MIPMAP_LEVELS];
+	TArray<FStructuredBufferPoolItemSharedPtr> BufferDataArray;
+	BufferDataArray.AddDefaulted(FImgMediaLoader::MAX_MIPMAP_LEVELS);
 
 	int32 NumMipLevels = Loader->GetNumMipLevels();
 	{
 		// Loop over all mips.
-		FIntPoint Dim = FullResolution;
+		FIntPoint CurrentMipDim = FullResolution;
 
 		for (int32 CurrentMipLevel = 0; CurrentMipLevel < NumMipLevels; ++CurrentMipLevel)
 		{
@@ -175,103 +185,58 @@ bool FExrImgMediaReaderGpu::ReadFrame(int32 FrameId, int32 MipLevel, const FImgM
 			bool ReadThisMip = (CurrentMipLevel >= MipLevel) && (IsThisLevelPresent == false);
 
 			// Next mip level.
-			Dim = FullResolution / FMath::Pow(2., float(CurrentMipLevel));
+			int MipPow = FMath::Pow(2., float(CurrentMipLevel));
+			CurrentMipDim = FullResolution / MipPow;
 
 			if (ReadThisMip)
 			{
-				const uint8 PlanePadding = 8;
-				const SIZE_T BufferSize = GetBufferSize(Dim, NumChannels) - Dim.Y * PlanePadding + NumTilesX * NumTilesY * TileDim.Y * PlanePadding;
+				const SIZE_T BufferSize = GetBufferSize(CurrentMipDim, NumChannels, bHasTiles, OutFrame->Info.NumTiles / MipPow);
 				FStructuredBufferPoolItemSharedPtr& BufferData = BufferDataArray[CurrentMipLevel];
 				BufferData = AllocateGpuBufferFromPool(BufferSize);
 				uint16* MipDataPtr = static_cast<uint16*>(BufferData->MappedBuffer);
 
 				// Get for our frame/mip level.
 				FString ImagePath = Loader->GetImagePath(FrameId, CurrentMipLevel);
-				FString BaseImage;
-				if (bHasTiles)
-				{
-					// Remove "_x0_y0.exr" so we can add on the correct name for the tile we want.
-					BaseImage = ImagePath.LeftChop(10);
-				}
+				
 				EReadResult ReadResult = Fail;
 
-				int32 TileWidth = Dim.X / NumTilesX;
-				int32 TileHeight = Dim.Y / NumTilesY;
-
-				FIntPoint TileSize(TileWidth, TileHeight);
-				const SIZE_T DataSizeToRead = GetBufferSize(TileSize, NumChannels);
-
 #if READ_IN_CHUNKS
-				for (int32 TileY = StartTileY; TileY < EndTileY; TileY++)
+				FRgbaInputFile InputTileFile(ImagePath);
+				if (InputTileFile.HasInputFile())
 				{
-					for (int32 TileX = StartTileX; TileX < EndTileX; TileX++)
-					{
-						// Get for our frame/mip level.
-						if (bHasTiles)
-						{
-							ImagePath = FString::Printf(TEXT("%s_x%d_y%d.exr"), *BaseImage, TileX, TileY);
-						}
-						FRgbaInputFile InputTileFile(ImagePath);
-						if (InputTileFile.HasInputFile())
-						{
-							// read frame data
-							uint32 CurrentBufferOffset = (NumTilesX * TileY + TileX) * DataSizeToRead / sizeof(uint16);
-							check(CurrentBufferOffset <= BufferSize);
-
-							ReadResult = ReadInChunks(MipDataPtr + CurrentBufferOffset, ImagePath, FrameId, TileSize, DataSizeToRead, PixelSize, NumChannels);
+					// read frame data
+					ReadResult = ReadInChunks(MipDataPtr, ImagePath, FrameId, CurrentMipDim, BufferSize, PixelSize, NumChannels);
 							
-							if (ReadResult == Fail)
-							{
-								goto ReadFailedFallback;
-							}
-
-							OutFrame->MipMapsPresent |= 1 << CurrentMipLevel;
-						}
-						else
-						{
-							UE_LOG(LogImgMedia, Error, TEXT("Could not load %s"), *ImagePath);
-							return false;
-						}
+					if (ReadResult != Fail)
+					{
+						OutFrame->MipMapsPresent |= 1 << CurrentMipLevel;
 					}
 				}
+				else
+				{
+					UE_LOG(LogImgMedia, Error, TEXT("Could not load %s"), *ImagePath);
+					return false;
+				}
 #else
-				for (int32 TileY = StartTileY; TileY < EndTileY; TileY++)
+				FRgbaInputFile InputTileFile(ImagePath);
+				if (InputTileFile.HasInputFile())
 				{
-					for (int32 TileX = StartTileX; TileX < EndTileX; TileX++)
+					// read frame data
+					bool bResult = FExrReader::GenerateTextureData(MipDataPtr, BufferSize, ImagePath, CurrentMipDim.Y, NumChannels);
+					ReadResult = bResult ? Success : Fail;
+					if (bResult)
 					{
-						// Get for our frame/mip level.
-						if (bHasTiles)
-						{
-							ImagePath = FString::Printf(TEXT("%s_x%d_y%d.exr"), *BaseImage, TileX, TileY);
-						}
-						FRgbaInputFile InputTileFile(ImagePath);
-						if (InputTileFile.HasInputFile())
-						{
-							// read frame data
-							uint32 CurrentBufferOffset = (NumTilesX * (TileY)+TileX) * DataSizeToRead / sizeof(uint16);
-							check(CurrentBufferOffset <= BufferSize);
-							bool bResult = FExrReader::GenerateTextureData(MipDataPtr + CurrentBufferOffset, ImagePath, TileSize.X, TileSize.Y, PixelSize, NumChannels);
-							ReadResult = bResult ? Success : Fail;
-
-							if (!bResult)
-							{
-								goto ReadFailedFallback;
-							}
-
-							OutFrame->MipMapsPresent |= 1 << CurrentMipLevel;
-						}
-						else
-						{
-							UE_LOG(LogImgMedia, Error, TEXT("Could not load %s"), *ImagePath);
-							return false;
-						}
+						OutFrame->MipMapsPresent |= 1 << CurrentMipLevel;
 					}
 				}
+				else
+				{
+					UE_LOG(LogImgMedia, Error, TEXT("Could not load %s"), *ImagePath);
+					return false;
+				}
 #endif
-
 				if (ReadResult == Fail)
 				{
-				ReadFailedFallback:
 					// Check if we have a compressed file.
 					FRgbaInputFile InputFileMip(ImagePath);
 					FImgMediaFrameInfo Info;
@@ -291,77 +256,19 @@ bool FExrImgMediaReaderGpu::ReadFrame(int32 FrameId, int32 MipLevel, const FImgM
 			}
 		}
 	}
-
 	OutFrame->Format = NumChannels <= 3 ? EMediaTextureSampleFormat::FloatRGB : EMediaTextureSampleFormat::FloatRGBA;
 	OutFrame->Stride = FullResolution.X * PixelSize;
-	auto RenderThreadSwizzler = [this, BufferDataArray, FullResolution, TileDim, Viewport, NumChannels, NumMipLevels, bHasTiles](FRHICommandListImmediate& RHICmdList, FTexture2DRHIRef RenderTargetTextureRHI)->bool
-	{
-		SCOPED_DRAW_EVENT(RHICmdList, FExrImgMediaReaderGpu_Convert);
-		SCOPED_GPU_STAT(RHICmdList, ExrImgMediaReaderGpu);
-
-		FIntPoint Dim = FullResolution;
-		for (int32 MipLevel = 0; MipLevel < NumMipLevels; ++MipLevel)
-		{
-			FStructuredBufferPoolItemSharedPtr BufferData = BufferDataArray[MipLevel];
-			if (BufferData.IsValid())
-			{
-				if (!BufferData->BufferRef->IsValid())
-				{
-					continue;
-				}
-				// This flag will indicate that we should wait for poll to complete.
-				BufferData->bWillBeSignaled = true;
 
-				FRHIRenderPassInfo RPInfo(RenderTargetTextureRHI, ERenderTargetActions::DontLoad_Store, nullptr, MipLevel);
-				RHICmdList.BeginRenderPass(RPInfo, TEXT("ExrTextureSwizzle"));
+	OutFrame->SampleConverter = CreateSampleConverter(MoveTemp(BufferDataArray), FullResolution, TileDim, TopMipViewport, NumChannels, NumMipLevels, bHasTiles);
 
-				FExrSwizzlePS::FPermutationDomain PermutationVector;
-				PermutationVector.Set<FExrSwizzlePS::FRgbaSwizzle>(NumChannels);
-				PermutationVector.Set<FExrSwizzlePS::FSwizzleTiles>(bHasTiles);
-				FExrSwizzlePS::FParameters Parameters = FExrSwizzlePS::FParameters();
-				Parameters.TextureSize = Dim;
-				Parameters.TileSize = TileDim;
-
-				Parameters.UnswizzledBuffer = RHICreateShaderResourceView(BufferData->BufferRef);
-
-				FGlobalShaderMap* ShaderMap = GetGlobalShaderMap(GMaxRHIFeatureLevel);
-
-				TShaderMapRef<FExrSwizzleVS> SwizzleShaderVS(ShaderMap);
-				TShaderMapRef<FExrSwizzlePS> SwizzleShaderPS(ShaderMap, PermutationVector);
-
-				FScreenPassPipelineState PipelineState(SwizzleShaderVS, SwizzleShaderPS, TStaticBlendState<>::GetRHI(), TStaticDepthStencilState<false, CF_Always>::GetRHI());
-				DrawScreenPass(RHICmdList, Dim, Viewport, PipelineState, [&](FRHICommandListImmediate& RHICmdList)
-				{
-					SetShaderParameters(RHICmdList, SwizzleShaderPS, SwizzleShaderPS.GetPixelShader(), Parameters);
-				});
-
-				// Resolve render target.
-				RHICmdList.EndRenderPass();
-
-				// Mark this render command for this buffer as complete, so we can poll it and transfer later.
-				RHICmdList.WriteGPUFence(BufferData->Fence);
-			}
-
-			// Next level.
-			Dim /= 2;
-		}
-
-		//Doesn't need further conversion so returning false.
-		return false;
-	};
-
-	FExrMediaTextureSampleConverter* SampleConverter = new FExrMediaTextureSampleConverter();
-	SampleConverter->ConvertExrBufferCallback = FExrConvertBufferCallback::CreateLambda(RenderThreadSwizzler);
-	OutFrame->SampleConverter = MakeShareable(SampleConverter);
 	OutFrame->MipMapsPresent = 1 << MipLevel;
 	UE_LOG(LogImgMedia, Verbose, TEXT("Reader %p: Read Pixels Complete. %i"), this, FrameId);
-
 	return true;
 }
 
 void FExrImgMediaReaderGpu::PreAllocateMemoryPool(int32 NumFrames, const FImgMediaFrameInfo& FrameInfo)
 {
-	SIZE_T AllocSize = GetBufferSize(FrameInfo.Dim, FrameInfo.NumChannels);
+	SIZE_T AllocSize = GetBufferSize(FrameInfo.Dim, FrameInfo.NumChannels, FrameInfo.bHasTiles, FrameInfo.NumTiles);
 	for (int32 FrameCacheNum = 0; FrameCacheNum < NumFrames; FrameCacheNum++)
 	{
 		AllocateGpuBufferFromPool(AllocSize, FrameCacheNum == NumFrames - 1);
@@ -435,14 +342,106 @@ FExrImgMediaReaderGpu::EReadResult FExrImgMediaReaderGpu::ReadInChunks(uint16* B
 	return bResult;
 }
 
-SIZE_T FExrImgMediaReaderGpu::GetBufferSize(const FIntPoint& Dim, int32 NumChannels)
+SIZE_T FExrImgMediaReaderGpu::GetBufferSize(const FIntPoint& Dim, int32 NumChannels, bool bHasTiles, const FIntPoint& TileNum)
 {
-	// At the beginning of each row of B G R channel planes there is 2x4 byte data that has information
-	// about number of pixels in the current row and row's number.
-	const uint16 PlanePadding = 8;
+	if (!bHasTiles)
+	{
+		/** 
+		* Reading scanlines.
+		* 
+		* At the beginning of each row of B G R channel planes there is 2x4 byte data that has information
+		* about number of pixels in the current row and row's number.
+		*/
+		SIZE_T BufferSize = Dim.X * Dim.Y * sizeof(uint16) * NumChannels + Dim.Y * ExrRowPadding;
+		return BufferSize;
+	}
+	else
+	{
+		/** 
+		* Reading tiles.
+		* 
+		* At the beginning of each tile of B G R channel planes there is 10 byte data that has information
+		* about number contents of tiles.
+		*/
+		SIZE_T BufferSize = Dim.X * Dim.Y * sizeof(uint16) * NumChannels + TileNum.X * TileNum.Y * ExrTilePadding;
+		return BufferSize;
+	}
+
+}
+
+TSharedPtr<IMediaTextureSampleConverter, ESPMode::ThreadSafe> FExrImgMediaReaderGpu::CreateSampleConverter
+	( TArray<FStructuredBufferPoolItemSharedPtr>&& BufferDataArray
+	, const FIntPoint& FullResolution
+	, const FIntPoint& TileDim
+	, const FIntRect& TopMipViewport
+	, const int32 NumChannels
+	, const int32 NumMipLevels
+	, const bool bHasTiles)
+{
+	auto RenderThreadSwizzler = [this, BufferDataArray, FullResolution, TileDim, TopMipViewport, NumChannels, NumMipLevels, bHasTiles](FRHICommandListImmediate& RHICmdList, FTexture2DRHIRef RenderTargetTextureRHI)->bool
+	{
+
+		SCOPED_DRAW_EVENT(RHICmdList, FExrImgMediaReaderGpu_Convert);
+		SCOPED_GPU_STAT(RHICmdList, ExrImgMediaReaderGpu);
+
+		FIntPoint Dim = FullResolution;
+		for (int32 MipLevel = 0; MipLevel < NumMipLevels; ++MipLevel)
+		{
+			int MipLevelDiv = FMath::Pow(2., float(MipLevel));
+			Dim = FullResolution / MipLevelDiv;
+			const FIntRect Viewport = TopMipViewport.Scale(1. / MipLevelDiv);
+
+			FStructuredBufferPoolItemSharedPtr BufferData = BufferDataArray[MipLevel];
+			if (BufferData.IsValid())
+			{
+				if (!BufferData->BufferRef->IsValid())
+				{
+					continue;
+				}
+				// This flag will indicate that we should wait for poll to complete.
+				BufferData->bWillBeSignaled = true;
+
+				FRHIRenderPassInfo RPInfo(RenderTargetTextureRHI, ERenderTargetActions::DontLoad_Store, nullptr, MipLevel);
+				RHICmdList.BeginRenderPass(RPInfo, TEXT("ExrTextureSwizzle"));
+
+				FExrSwizzlePS::FPermutationDomain PermutationVector;
+				PermutationVector.Set<FExrSwizzlePS::FRgbaSwizzle>(NumChannels);
+				PermutationVector.Set<FExrSwizzlePS::FSwizzleTiles>(bHasTiles);
+				FExrSwizzlePS::FParameters Parameters = FExrSwizzlePS::FParameters();
+				Parameters.TextureSize = Dim;
+				Parameters.TileSize = TileDim;
+
+				Parameters.UnswizzledBuffer = RHICreateShaderResourceView(BufferData->BufferRef);
+
+				FGlobalShaderMap* ShaderMap = GetGlobalShaderMap(GMaxRHIFeatureLevel);
+
+				TShaderMapRef<FExrSwizzleVS> SwizzleShaderVS(ShaderMap);
+				TShaderMapRef<FExrSwizzlePS> SwizzleShaderPS(ShaderMap, PermutationVector);
+
+				FScreenPassPipelineState PipelineState(SwizzleShaderVS, SwizzleShaderPS, TStaticBlendState<>::GetRHI(), TStaticDepthStencilState<false, CF_Always>::GetRHI());
+				DrawScreenPass(RHICmdList, Dim, Viewport, PipelineState, [&](FRHICommandListImmediate& RHICmdList)
+					{
+						SetShaderParameters(RHICmdList, SwizzleShaderPS, SwizzleShaderPS.GetPixelShader(), Parameters);
+					});
+
+				// Resolve render target.
+				RHICmdList.EndRenderPass();
+
+				// Mark this render command for this buffer as complete, so we can poll it and transfer later.
+				RHICmdList.WriteGPUFence(BufferData->Fence);
+			}
+
+			
+		}
+
+		//Doesn't need further conversion so returning false.
+		return false;
+	};
+
+	FExrMediaTextureSampleConverter* SampleConverter = new FExrMediaTextureSampleConverter();
+	SampleConverter->ConvertExrBufferCallback = FExrConvertBufferCallback::CreateLambda(RenderThreadSwizzler);
 
-	SIZE_T BufferSize = Dim.X * Dim.Y * sizeof(uint16) * NumChannels + Dim.Y * PlanePadding;
-	return BufferSize;
+	return MakeShareable(SampleConverter);
 }
 
 FStructuredBufferPoolItemSharedPtr FExrImgMediaReaderGpu::AllocateGpuBufferFromPool(uint32 AllocSize, bool bWait)
diff --git a/Engine/Plugins/Media/ImgMedia/Source/ImgMedia/Private/Readers/ExrImgMediaReaderGpu.h b/Engine/Plugins/Media/ImgMedia/Source/ImgMedia/Private/Readers/ExrImgMediaReaderGpu.h
--- a/Engine/Plugins/Media/ImgMedia/Source/ImgMedia/Private/Readers/ExrImgMediaReaderGpu.h
+++ b/Engine/Plugins/Media/ImgMedia/Source/ImgMedia/Private/Readers/ExrImgMediaReaderGpu.h
@@ -92,7 +92,17 @@ protected:
 	 * @param Dim Dimensions of the image.
 	 * @param NumChannels Number of channels in the image.
 	 */
-	static SIZE_T GetBufferSize(const FIntPoint& Dim, int32 NumChannels);
+	static SIZE_T GetBufferSize(const FIntPoint& Dim, int32 NumChannels, bool bHasTiles, const FIntPoint& TileNum);
+
+
+	TSharedPtr<IMediaTextureSampleConverter, ESPMode::ThreadSafe> CreateSampleConverter
+		( TArray<FStructuredBufferPoolItemSharedPtr>&& BufferDataArray
+		, const FIntPoint& FullResolution
+		, const FIntPoint& TileDim
+		, const FIntRect& TopMipViewport
+		, const int32 NumChannels
+		, const int32 NumMipLevels
+		, const bool bHasTiles);
 
 public:
 
