ImgMedia: Process Images can write out tiled images in our custom format.

#rnx

diff --git a/Engine/Plugins/Media/ImgMedia/Source/ImgMediaEditor/Private/Widgets/SImgMediaProcessImages.cpp b/Engine/Plugins/Media/ImgMedia/Source/ImgMediaEditor/Private/Widgets/SImgMediaProcessImages.cpp
--- a/Engine/Plugins/Media/ImgMedia/Source/ImgMediaEditor/Private/Widgets/SImgMediaProcessImages.cpp
+++ b/Engine/Plugins/Media/ImgMedia/Source/ImgMediaEditor/Private/Widgets/SImgMediaProcessImages.cpp
@@ -281,6 +281,40 @@ void SImgMediaProcessImages::ProcessImageCustom(TSharedPtr<IImageWrapper>& InIma
 	int32 BytesPerPixel = RawData.Num() / (Width * Height);
 
 	bool bIsTiled = (NumTilesX > 1) || (NumTilesY > 1);
+	uint8* RawDataPtr = RawData.GetData();
+
+	// Tile data.
+	TArray64<uint8> RawDataTiled;
+	if (bIsTiled)
+	{
+		RawDataTiled.AddUninitialized(RawData.Num());
+		RawDataPtr = RawDataTiled.GetData();
+
+		uint8* SourceData = RawData.GetData();
+		uint8* DestData = RawDataTiled.GetData();
+		int32 BytesPerTile = TileWidth * TileHeight * BytesPerPixel;
+
+		// Loop over y tiles.
+		for (int32 TileY = 0; TileY < NumTilesY; ++TileY)
+		{
+			// Loop over x tiles.
+			for (int32 TileX = 0; TileX < NumTilesX; ++TileX)
+			{
+				// Get address of the source and destination tiles.
+				uint8* SourceTile = SourceData +
+					(TileX * TileWidth + TileY * Width * TileHeight) * BytesPerPixel;
+				uint8* DestTile = DestData + (TileX + TileY * NumTilesX) * BytesPerTile;
+				
+				// Loop over each row in the tile.
+				for (int32 Row = 0; Row < TileHeight; ++Row)
+				{
+					uint8* SourceLine = SourceTile + Row * Width * BytesPerPixel;
+					uint8* DestLine = DestTile + Row * TileWidth * BytesPerPixel;
+					FMemory::Memcpy(DestLine, SourceLine, TileWidth * BytesPerPixel);
+				}
+			}
+		}
+	}
 
 	// Names for our channels.
 	const FString RChannelName = FString(TEXT("R"));
@@ -310,23 +344,13 @@ void SImgMediaProcessImages::ProcessImageCustom(TSharedPtr<IImageWrapper>& InIma
 	OutFile.AddChannel(RChannelName);
 
 	// Create output.
-	OutFile.CreateOutputFile(InName, TileWidth, TileHeight, false);
-	OutFile.AddFrameBufferChannel(AChannelName, RawData.GetData(), Stride);
-	OutFile.AddFrameBufferChannel(BChannelName, RawData.GetData() + Width * 2, Stride);
-	OutFile.AddFrameBufferChannel(GChannelName, RawData.GetData() + Width * 4, Stride);
-	OutFile.AddFrameBufferChannel(RChannelName, RawData.GetData() + Width * 6, Stride);
+	OutFile.CreateOutputFile(InName, Width, Height, false);
+	OutFile.AddFrameBufferChannel(AChannelName, RawDataPtr, Stride);
+	OutFile.AddFrameBufferChannel(BChannelName, RawDataPtr + Width * 2, Stride);
+	OutFile.AddFrameBufferChannel(GChannelName, RawDataPtr + Width * 4, Stride);
+	OutFile.AddFrameBufferChannel(RChannelName, RawDataPtr + Width * 6, Stride);
 	OutFile.SetFrameBuffer();
 	
-	// Loop over y tiles.
-	for (int TileY = 0; TileY < NumTilesY; ++TileY)
-	{
-		// Loop over x tiles.
-		for (int TileX = 0; TileX < NumTilesX; ++TileX)
-		{
-			
-		}
-	}
-	
 	OutFile.WriteTile(0, 0, 0);
 
 #else // IMGMEDIAEDITOR_EXR_SUPPORTED_PLATFORM
