ImgMedia ExrReaderGpu: Reader now reads tile regions instead of full file.

diff --git a/Engine/Plugins/Media/ImgMedia/Shaders/Private/ExrSwizzler.usf b/Engine/Plugins/Media/ImgMedia/Shaders/Private/ExrSwizzler.usf
--- a/Engine/Plugins/Media/ImgMedia/Shaders/Private/ExrSwizzler.usf
+++ b/Engine/Plugins/Media/ImgMedia/Shaders/Private/ExrSwizzler.usf
@@ -11,6 +11,8 @@ int2 TextureSize;
 /** Size of tiles that represent this texture. */
 int2 TileSize;
 
+/** Total number of tiles in X and Y direction for the current mip level. */
+int2 NumTiles;
 
 /**
 * This function extracts the 16bit (float) data stored in uint (32 bit)
@@ -56,7 +58,7 @@ void MainVS(
 	DrawRectangle(InPosition, InTexCoord, OutPosition, OutUVAndScreenPos);
 }
 
-void GetStartScanlinePosition(in int2 PixelPos, out float2 SampleSize, out int StartPosition)
+void GetStartScanlinePosition(in int2 PixelPos, out int2 SampleSize, out int StartPosition)
 {
 	// Depending on the EXR file it could contain multi-channel data.
 #if NUM_CHANNELS == 1
@@ -70,8 +72,7 @@ void GetStartScanlinePosition(in int2 PixelPos, out float2 SampleSize, out int S
 #endif
 
 #if SWIZZLE_TILES
-	SampleSize.x = TileSize.x;
-	SampleSize.y = TileSize.y;
+	SampleSize = TileSize;
 
 	int xCoord = (PixelPos.x) % TileSize.x;
 	int yCoord = (PixelPos.y) % TileSize.y;
@@ -79,16 +80,13 @@ void GetStartScanlinePosition(in int2 PixelPos, out float2 SampleSize, out int S
 	int TileX = floor((PixelPos.x) / TileSize.x);
 	int TileY = floor((PixelPos.y) / TileSize.y);
 
-	int MaxTilesX = TextureSize.x / TileSize.x;
-	int MaxTilesY = TextureSize.y / TileSize.y;
-
 	// Padding in bytes
 	const int TilePadding = 10;
 
 	// Current tile position.
 	int TileBufferStride = TileSize.x * TileSize.y * Channels;
-	int PreviousTilesOffset = (MaxTilesX * TileY + TileX) * TileBufferStride;
-	StartPosition = PreviousTilesOffset + yCoord * TileSize.x * Channels + xCoord + (TileY * MaxTilesX + TileX + 1) * TilePadding;
+	int PreviousTilesOffset = (NumTiles.x * TileY + TileX) * TileBufferStride;
+	StartPosition = PreviousTilesOffset + yCoord * TileSize.x * Channels + xCoord + (TileY * NumTiles.x + TileX + 1) * TilePadding;
 #else
 	SampleSize.x = TextureSize.x;
 	SampleSize.y = TextureSize.y;
@@ -106,7 +104,7 @@ half4 MainPS(noperspective float4 UVAndScreenPos : TEXCOORD0) : SV_Target0
 {
 	float2 UV = UVAndScreenPos.xy;
 
-	float2 SampleSize;
+	int2 SampleSize;
 
 	int2 PixelPos = UV * TextureSize;
 	int StartBufferPosition;
diff --git a/Engine/Plugins/Media/ImgMedia/Source/ExrReaderGpu/Private/ExrReaderGpu.cpp b/Engine/Plugins/Media/ImgMedia/Source/ExrReaderGpu/Private/ExrReaderGpu.cpp
--- a/Engine/Plugins/Media/ImgMedia/Source/ExrReaderGpu/Private/ExrReaderGpu.cpp
+++ b/Engine/Plugins/Media/ImgMedia/Source/ExrReaderGpu/Private/ExrReaderGpu.cpp
@@ -199,7 +199,7 @@ bool FExrReader::ReadHeaderData(FILE* FileHandle)
 	return true;
 }
 
-bool FExrReader::ReadLineOffsets(FILE* FileHandle, ELineOrder LineOrder, TArray<int64>& LineOffsets)
+bool FExrReader::ReadLineOrTileOffsets(FILE* FileHandle, ELineOrder LineOrder, TArray<int64>& LineOrTileOffsets)
 {
 	// At the moment we only support increasing Y EXR; 
 	// LineOrder is currently unused but we might need to take that into account as exr scanlines can go from top to bottom and vice versa.
@@ -209,9 +209,9 @@ bool FExrReader::ReadLineOffsets(FILE* FileHandle, ELineOrder LineOrder, TArray<
 		return false;
 	}
 
-	for (int32 i = 0; i < LineOffsets.Num(); i++)
+	for (int32 i = 0; i < LineOrTileOffsets.Num(); i++)
 	{
-		Read8ByteValue(FileHandle, LineOffsets[i]);
+		Read8ByteValue(FileHandle, LineOrTileOffsets[i]);
 	}
 	return true;
 }
@@ -233,25 +233,25 @@ bool FExrReader::GenerateTextureData(uint16* Buffer, int32 BufferSize, FString F
 	rewind(FileHandle);
 
 	// Reading header (throwaway) and line offset data. We just need line offset data.
-	TArray<int64> LineOffsets;
+	TArray<int64> LineOrTileOffsets;
 	{
 		ReadMagicNumberAndVersionField(FileHandle);
 		ReadHeaderData(FileHandle);
-		LineOffsets.SetNum(NumberOfScanlines);
+		LineOrTileOffsets.SetNum(NumberOfScanlines);
 
 		// At the moment we support only increasing Y.
 		ELineOrder LineOrder = INCREASING_Y;
-		ReadLineOffsets(FileHandle, LineOrder, LineOffsets);
+		ReadLineOrTileOffsets(FileHandle, LineOrder, LineOrTileOffsets);
 	}
 
-	fseek(FileHandle, LineOffsets[0], SEEK_SET);
+	fseek(FileHandle, LineOrTileOffsets[0], SEEK_SET);
 	fread(Buffer, BufferSize, 1 /*NumOfElements*/, FileHandle);
 
 	fclose(FileHandle);
 	return true;
 }
 
-bool FExrReader::OpenExrAndPrepareForPixelReading(FString FilePath, int32 TextureWidth, int32 TextureHeight, int32& PixelSize, int32 NumChannels)
+bool FExrReader::OpenExrAndPrepareForPixelReading(FString FilePath, int32 NumOffsets)
 {
 	if (FileHandle != nullptr)
 	{
@@ -273,14 +273,14 @@ bool FExrReader::OpenExrAndPrepareForPixelReading(FString FilePath, int32 Textur
 		{
 			return false;
 		}
-		LineOffsets.SetNum(TextureHeight);
+		LineOrTileOffsets.SetNum(NumOffsets);
 
 		// At the moment we support only increasing Y.
 		ELineOrder LineOrder = INCREASING_Y;
-		ReadLineOffsets(FileHandle, LineOrder, LineOffsets);
+		ReadLineOrTileOffsets(FileHandle, LineOrder, LineOrTileOffsets);
 	}
 
-	fseek(FileHandle, LineOffsets[0], SEEK_SET);
+	fseek(FileHandle, LineOrTileOffsets[0], SEEK_SET);
 
 	return true;
 }
@@ -289,7 +289,7 @@ bool FExrReader::ReadExrImageChunk(void* Buffer, int32 ChunkSize)
 {
 	if (FileHandle == nullptr)
 	{
-		UE_LOG(LogExrReaderGpu, Error, TEXT("File is not open for reading. Please use OpenExrAndPrepareForPixelReading"));
+		UE_LOG(LogExrReaderGpu, Error, TEXT("File is not open for reading. Please use OpenExrAndPrepareForPixelReading. "));
 		return false;
 	}
 	if (Buffer == nullptr)
@@ -306,6 +306,18 @@ bool FExrReader::ReadExrImageChunk(void* Buffer, int32 ChunkSize)
 	return NumElementsRead == NumElements;
 }
 
+bool FExrReader::SeekTileWithinFile(const int StartTileIndex, const FIntPoint& TexDimInTiles, int32& OutBufferOffset)
+{
+	if (StartTileIndex >= LineOrTileOffsets.Num())
+	{
+		UE_LOG(LogExrReaderGpu, Error, TEXT("Tile index is invalid."));
+		return false;
+	}
+	int64 LineOffset = LineOrTileOffsets[StartTileIndex];
+	OutBufferOffset = LineOffset - LineOrTileOffsets[0];
+	return fseek(FileHandle, LineOffset, SEEK_SET) == 0;
+}
+
 bool FExrReader::CloseExrFile()
 {
 	if (FileHandle == nullptr)
@@ -314,7 +326,7 @@ bool FExrReader::CloseExrFile()
 		return false;
 	}
 	fclose(FileHandle);
-	LineOffsets.Empty();
+	LineOrTileOffsets.Empty();
 	return true;
 }
 
diff --git a/Engine/Plugins/Media/ImgMedia/Source/ExrReaderGpu/Public/ExrReaderGpu.h b/Engine/Plugins/Media/ImgMedia/Source/ExrReaderGpu/Public/ExrReaderGpu.h
--- a/Engine/Plugins/Media/ImgMedia/Source/ExrReaderGpu/Public/ExrReaderGpu.h
+++ b/Engine/Plugins/Media/ImgMedia/Source/ExrReaderGpu/Public/ExrReaderGpu.h
@@ -46,7 +46,7 @@ private:
 	static bool ReadHeaderData(FILE* FileHandle);
 
 	/** Reads Scan line offsets so that we can jump to a position in file. */
-	static bool ReadLineOffsets(FILE* FileHandle, ELineOrder LineOrder, TArray<int64>& LineOffsets);
+	static bool ReadLineOrTileOffsets(FILE* FileHandle, ELineOrder LineOrder, TArray<int64>& LineOrTileOffsets);
 
 public:
 
@@ -58,7 +58,7 @@ public:
 	* Reading in chunks allows the process of reading to be canceled midway through reading. 
 	* This function reads and discards the header and keeps the pointers to the scanlines. 
 	*/
-	bool OpenExrAndPrepareForPixelReading(FString FilePath, int32 TextureWidth, int32 TextureHeight, int32& PixelSize, int32 NumChannels);
+	bool OpenExrAndPrepareForPixelReading(FString FilePath, int32 NumOffsets);
 
 	/** 
 	* Read a chunk of an Exr file previously open via OpenExrAndPrepareForPixelReading.
@@ -66,6 +66,11 @@ public:
 	*/
 	bool ReadExrImageChunk(void* Buffer, int32 ChunkSize);
 
+	/** 
+	* Moves to the specified tile position in the file in preparation for reading.
+	*/
+	bool SeekTileWithinFile(const int StartTileIndex, const FIntPoint& TexDimInTiles, int32& OutBufferOffset);
+
 	/**
 	* After we have finished reading the file that was open via OpenExrAndPrepareForPixelReading
 	* We need to close it to avoid memory leaks
@@ -80,9 +85,8 @@ private:
 	FILE* FileHandle;
 
 	/**
-	* These are byte offsets pointing to scanlines from the begining of the file.
+	* These are byte offsets pointing to scanlines or tiles from the begining of the file.
 	*/
-	TArray<int64> LineOffsets;
-
+	TArray<int64> LineOrTileOffsets;
 };
 #endif
diff --git a/Engine/Plugins/Media/ImgMedia/Source/ExrReaderGpu/Public/ExrSwizzlingShader.h b/Engine/Plugins/Media/ImgMedia/Source/ExrReaderGpu/Public/ExrSwizzlingShader.h
--- a/Engine/Plugins/Media/ImgMedia/Source/ExrReaderGpu/Public/ExrSwizzlingShader.h
+++ b/Engine/Plugins/Media/ImgMedia/Source/ExrReaderGpu/Public/ExrSwizzlingShader.h
@@ -53,6 +53,7 @@ class EXRREADERGPU_API FExrSwizzlePS : public FGlobalShader
 		SHADER_PARAMETER_SRV(StructuredBuffer<uint>, UnswizzledBuffer)
 		SHADER_PARAMETER(FIntPoint, TextureSize)
 		SHADER_PARAMETER(FIntPoint, TileSize)
+		SHADER_PARAMETER(FIntPoint, NumTiles)
 	END_SHADER_PARAMETER_STRUCT()
 };
 #endif
\ No newline at end of file
diff --git a/Engine/Plugins/Media/ImgMedia/Source/ImgMedia/Private/Readers/ExrImgMediaReaderGpu.cpp b/Engine/Plugins/Media/ImgMedia/Source/ImgMedia/Private/Readers/ExrImgMediaReaderGpu.cpp
--- a/Engine/Plugins/Media/ImgMedia/Source/ImgMedia/Private/Readers/ExrImgMediaReaderGpu.cpp
+++ b/Engine/Plugins/Media/ImgMedia/Source/ImgMedia/Private/Readers/ExrImgMediaReaderGpu.cpp
@@ -35,10 +35,10 @@ DECLARE_GPU_STAT_NAMED(ExrImgMediaReaderGpu, TEXT("ExrImgMediaReaderGpu"));
 namespace {
 
 	/* Padding in bits that is added at the beginning of each scanline */
-	const uint16 ExrRowPadding = 8;
+	const uint16 CExrRowPadding = 8;
 
 	/* Padding in bits that is added at the beginning of each tile */
-	const uint16 ExrTilePadding = 20;
+	const uint16 CExrTilePadding = 20;
 
 	/** This function is similar to DrawScreenPass in OpenColorIODisplayExtension.cpp except it is catered for Viewless texture rendering. */
 	template<typename TSetupFunction>
@@ -166,7 +166,6 @@ bool FExrImgMediaReaderGpu::ReadFrame(int32 FrameId, int32 MipLevel, const FImgM
 	{
 		TopMipViewport = FIntRect(StartTileX * TileDim.X, StartTileY * TileDim.Y, EndTileX * TileDim.X, EndTileY * TileDim.Y);
 	}
-
 	const int32 NumChannels = OutFrame->Info.NumChannels;
 	const int32 PixelSize = sizeof(uint16) * NumChannels;
 
@@ -185,12 +184,12 @@ bool FExrImgMediaReaderGpu::ReadFrame(int32 FrameId, int32 MipLevel, const FImgM
 			bool ReadThisMip = (CurrentMipLevel >= MipLevel) && (IsThisLevelPresent == false);
 
 			// Next mip level.
-			int MipPow = FMath::Pow(2., float(CurrentMipLevel));
-			CurrentMipDim = FullResolution / MipPow;
+			int MipLevelDiv = 1 << CurrentMipLevel;
+			CurrentMipDim = FullResolution / MipLevelDiv;
 
 			if (ReadThisMip)
 			{
-				const SIZE_T BufferSize = GetBufferSize(CurrentMipDim, NumChannels, bHasTiles, OutFrame->Info.NumTiles / MipPow);
+				const SIZE_T BufferSize = GetBufferSize(CurrentMipDim, NumChannels, bHasTiles, OutFrame->Info.NumTiles / MipLevelDiv);
 				FStructuredBufferPoolItemSharedPtr& BufferData = BufferDataArray[CurrentMipLevel];
 				BufferData = AllocateGpuBufferFromPool(BufferSize);
 				uint16* MipDataPtr = static_cast<uint16*>(BufferData->MappedBuffer);
@@ -205,8 +204,22 @@ bool FExrImgMediaReaderGpu::ReadFrame(int32 FrameId, int32 MipLevel, const FImgM
 				if (InputTileFile.HasInputFile())
 				{
 					// read frame data
-					ReadResult = ReadInChunks(MipDataPtr, ImagePath, FrameId, CurrentMipDim, BufferSize, PixelSize, NumChannels);
-							
+					if (bHasTiles)
+					{
+						FIntRect TileRegion = FIntRect(
+							FMath::FloorToInt(float(StartTileX) / MipLevelDiv),
+							FMath::FloorToInt(float(StartTileY) / MipLevelDiv),
+							FMath::CeilToInt(float(EndTileX) / MipLevelDiv),
+							FMath::CeilToInt(float(EndTileY) / MipLevelDiv));
+						FIntPoint NumTiles = OutFrame->Info.NumTiles / MipLevelDiv;
+
+						ReadResult = ReadTiles(MipDataPtr, ImagePath, FrameId, TileRegion, NumTiles, TileDim, PixelSize);
+					}
+					else
+					{
+						ReadResult = ReadInChunks(MipDataPtr, ImagePath, FrameId, CurrentMipDim, BufferSize);
+					}
+
 					if (ReadResult != Fail)
 					{
 						OutFrame->MipMapsPresent |= 1 << CurrentMipLevel;
@@ -290,7 +303,7 @@ void FExrImgMediaReaderGpu::OnTick()
 /* FExrImgMediaReaderGpu implementation
  *****************************************************************************/
 
-FExrImgMediaReaderGpu::EReadResult FExrImgMediaReaderGpu::ReadInChunks(uint16* Buffer, const FString& ImagePath, int32 FrameId, const FIntPoint& Dim, int32 BufferSize, int32 PixelSize, int32 NumChannels)
+FExrImgMediaReaderGpu::EReadResult FExrImgMediaReaderGpu::ReadInChunks(uint16* Buffer, const FString& ImagePath, int32 FrameId, const FIntPoint& Dim, int32 BufferSize)
 {
 	EReadResult bResult = Success;
 
@@ -302,7 +315,7 @@ FExrImgMediaReaderGpu::EReadResult FExrImgMediaReaderGpu::ReadInChunks(uint16* B
 	FExrReader ChunkReader;
 
 
-	if (!ChunkReader.OpenExrAndPrepareForPixelReading(ImagePath, Dim.X, Dim.Y, PixelSize, NumChannels))
+	if (!ChunkReader.OpenExrAndPrepareForPixelReading(ImagePath, Dim.Y))
 	{
 		return Fail;
 	}
@@ -342,6 +355,55 @@ FExrImgMediaReaderGpu::EReadResult FExrImgMediaReaderGpu::ReadInChunks(uint16* B
 	return bResult;
 }
 
+FExrImgMediaReaderGpu::EReadResult FExrImgMediaReaderGpu::ReadTiles(uint16* Buffer, const FString& ImagePath, int32 FrameId, const FIntRect& TileRegion, const FIntPoint& FullTexDimInTiles, const FIntPoint& TileDim, const int32 PixelSize)
+{
+	EReadResult bResult = Success;
+
+	// Chunks are of 16 MB
+	int32 CurrentBufferPos = 0;
+	FExrReader ChunkReader;
+
+
+	if (!ChunkReader.OpenExrAndPrepareForPixelReading(ImagePath, FullTexDimInTiles.X * FullTexDimInTiles.Y))
+	{
+		return Fail;
+	}
+
+	for (int32 TileRow = TileRegion.Min.Y; TileRow < TileRegion.Max.Y; TileRow++)
+	{
+		// Check to see if the frame was canceled.
+		{
+			FScopeLock RegionScopeLock(&CanceledFramesCriticalSection);
+			if (CanceledFrames.Remove(FrameId) > 0)
+			{
+				UE_LOG(LogImgMedia, Warning, TEXT("Reader %p: Canceling Frame %i At tile row # %i"), this, FrameId, TileRow);
+				bResult = Cancelled;
+				break;
+			}
+		}
+
+		const int32 TileByteStride = PixelSize * TileDim.X * TileDim.Y + CExrTilePadding;
+		const int StartTileIndex = TileRow * FullTexDimInTiles.X + TileRegion.Min.X;
+		const int EndTileIndex = TileRow * FullTexDimInTiles.X + TileRegion.Max.X;
+		const int BytesToRead = (EndTileIndex - StartTileIndex) * TileByteStride;
+
+		ChunkReader.SeekTileWithinFile(StartTileIndex, FullTexDimInTiles, CurrentBufferPos);
+
+		if (!ChunkReader.ReadExrImageChunk(reinterpret_cast<char*>(Buffer) + CurrentBufferPos, BytesToRead))
+		{
+			bResult = Fail;
+			break;
+		}
+	}
+
+	if (!ChunkReader.CloseExrFile())
+	{
+		return Fail;
+	}
+
+	return bResult;
+}
+
 SIZE_T FExrImgMediaReaderGpu::GetBufferSize(const FIntPoint& Dim, int32 NumChannels, bool bHasTiles, const FIntPoint& TileNum)
 {
 	if (!bHasTiles)
@@ -352,7 +414,7 @@ SIZE_T FExrImgMediaReaderGpu::GetBufferSize(const FIntPoint& Dim, int32 NumChann
 		* At the beginning of each row of B G R channel planes there is 2x4 byte data that has information
 		* about number of pixels in the current row and row's number.
 		*/
-		SIZE_T BufferSize = Dim.X * Dim.Y * sizeof(uint16) * NumChannels + Dim.Y * ExrRowPadding;
+		SIZE_T BufferSize = Dim.X * Dim.Y * sizeof(uint16) * NumChannels + Dim.Y * CExrRowPadding;
 		return BufferSize;
 	}
 	else
@@ -363,7 +425,7 @@ SIZE_T FExrImgMediaReaderGpu::GetBufferSize(const FIntPoint& Dim, int32 NumChann
 		* At the beginning of each tile of B G R channel planes there is 10 byte data that has information
 		* about number contents of tiles.
 		*/
-		SIZE_T BufferSize = Dim.X * Dim.Y * sizeof(uint16) * NumChannels + TileNum.X * TileNum.Y * ExrTilePadding;
+		SIZE_T BufferSize = Dim.X * Dim.Y * sizeof(uint16) * NumChannels + TileNum.X * TileNum.Y * CExrTilePadding;
 		return BufferSize;
 	}
 
@@ -387,10 +449,11 @@ TSharedPtr<IMediaTextureSampleConverter, ESPMode::ThreadSafe> FExrImgMediaReader
 		FIntPoint Dim = FullResolution;
 		for (int32 MipLevel = 0; MipLevel < NumMipLevels; ++MipLevel)
 		{
-			int MipLevelDiv = FMath::Pow(2., float(MipLevel));
+			int MipLevelDiv = 1 << MipLevel;
 			Dim = FullResolution / MipLevelDiv;
-			const FIntRect Viewport = TopMipViewport.Scale(1. / MipLevelDiv);
 
+			FIntRect Viewport = TopMipViewport / MipLevelDiv;
+			Viewport.Clip(FIntRect(FIntPoint(0, 0), Dim));
 			FStructuredBufferPoolItemSharedPtr BufferData = BufferDataArray[MipLevel];
 			if (BufferData.IsValid())
 			{
@@ -410,6 +473,10 @@ TSharedPtr<IMediaTextureSampleConverter, ESPMode::ThreadSafe> FExrImgMediaReader
 				FExrSwizzlePS::FParameters Parameters = FExrSwizzlePS::FParameters();
 				Parameters.TextureSize = Dim;
 				Parameters.TileSize = TileDim;
+				if (bHasTiles)
+				{
+					Parameters.NumTiles = Dim / TileDim;
+				}
 
 				Parameters.UnswizzledBuffer = RHICreateShaderResourceView(BufferData->BufferRef);
 
diff --git a/Engine/Plugins/Media/ImgMedia/Source/ImgMedia/Private/Readers/ExrImgMediaReaderGpu.h b/Engine/Plugins/Media/ImgMedia/Source/ImgMedia/Private/Readers/ExrImgMediaReaderGpu.h
--- a/Engine/Plugins/Media/ImgMedia/Source/ImgMedia/Private/Readers/ExrImgMediaReaderGpu.h
+++ b/Engine/Plugins/Media/ImgMedia/Source/ImgMedia/Private/Readers/ExrImgMediaReaderGpu.h
@@ -84,7 +84,13 @@ protected:
 	 * This function reads file in 16 MB chunks and if it detects that
 	 * Frame is pending for cancellation stops reading the file and returns false.
 	*/
-	EReadResult ReadInChunks(uint16* Buffer, const FString& ImagePath, int32 FrameId, const FIntPoint& Dim, int32 BufferSize, int32 PixelSize, int32 NumChannels);
+	EReadResult ReadInChunks(uint16* Buffer, const FString& ImagePath, int32 FrameId, const FIntPoint& Dim, int32 BufferSize);
+
+	/** 
+	 * Reads tiles from exr files in tile rows based on Tile region. If frame is pending for cancelation 
+	 * stops reading tiles at the current tile row.
+	*/
+	EReadResult ReadTiles(uint16* Buffer, const FString& ImagePath, int32 FrameId, const FIntRect& TileRegion, const FIntPoint& FullTexDimInTiles, const FIntPoint& TileDim, const int32 PixelSize);
 
 	/**
 	 * Get the size of the buffer needed to load in an image.
