ImgMedia: Process Images can write out duplicate pixels for a border for each tile.

#rnx

diff --git a/Engine/Plugins/Media/ImgMedia/Source/ImgMediaEditor/Private/Widgets/ImgMediaProcessImagesOptions.h b/Engine/Plugins/Media/ImgMedia/Source/ImgMediaEditor/Private/Widgets/ImgMediaProcessImagesOptions.h
--- a/Engine/Plugins/Media/ImgMedia/Source/ImgMediaEditor/Private/Widgets/ImgMediaProcessImagesOptions.h
+++ b/Engine/Plugins/Media/ImgMedia/Source/ImgMediaEditor/Private/Widgets/ImgMediaProcessImagesOptions.h
@@ -35,6 +35,10 @@ public:
 	UPROPERTY(EditAnywhere, Transient, Category = Tiles)
 	int32 TileSizeY = 0;
 
+	/** Number of pixels to duplicate along each tile border. */
+	UPROPERTY(EditAnywhere, Transient, Category = Tiles, meta = (ClampMin = "0.0"))
+	int32 TileBorder = 0;
+
 	/** Number of tiles in the X direction. If 0, then there are no tiles. */
 	UPROPERTY(VisibleAnywhere, Transient, Category = Tiles)
 	int32 NumTilesX = 0;
diff --git a/Engine/Plugins/Media/ImgMedia/Source/ImgMediaEditor/Private/Widgets/SImgMediaProcessImages.cpp b/Engine/Plugins/Media/ImgMedia/Source/ImgMediaEditor/Private/Widgets/SImgMediaProcessImages.cpp
--- a/Engine/Plugins/Media/ImgMedia/Source/ImgMediaEditor/Private/Widgets/SImgMediaProcessImages.cpp
+++ b/Engine/Plugins/Media/ImgMedia/Source/ImgMediaEditor/Private/Widgets/SImgMediaProcessImages.cpp
@@ -107,6 +107,7 @@ void SImgMediaProcessImages::ProcessAllImages(TSharedPtr<SNotificationItem> Conf
 	bool bUseCustomFormat = Options->bUseCustomFormat;
 	int32 InTileWidth = Options->TileSizeX;
 	int32 InTileHeight = Options->TileSizeY;
+	int32 TileBorder = Options->TileBorder;
 
 	// Create output directory.
 	FString OutPath = Options->OutputPath.Path;
@@ -175,7 +176,7 @@ void SImgMediaProcessImages::ProcessAllImages(TSharedPtr<SNotificationItem> Conf
 				FString Name = FPaths::Combine(OutPath, FileName);
 				if (bUseCustomFormat)
 				{
-					ProcessImageCustom(ImageWrapper, InTileWidth, InTileHeight, Name);
+					ProcessImageCustom(ImageWrapper, InTileWidth, InTileHeight, TileBorder, Name);
 				}
 				else
 				{
@@ -263,13 +264,15 @@ void SImgMediaProcessImages::ProcessImage(TSharedPtr<IImageWrapper>& InImageWrap
 }
 
 void SImgMediaProcessImages::ProcessImageCustom(TSharedPtr<IImageWrapper>& InImageWrapper,
-	int32 InTileWidth, int32 InTileHeight, const FString& InName)
+	int32 InTileWidth, int32 InTileHeight, int32 InTileBorder, const FString& InName)
 {
 #if IMGMEDIAEDITOR_EXR_SUPPORTED_PLATFORM
 	// Get image data.
 	ERGBFormat Format = InImageWrapper->GetFormat();
 	int32 Width = InImageWrapper->GetWidth();
 	int32 Height = InImageWrapper->GetHeight();
+	int32 DestWidth = Width;
+	int32 DestHeight = Height;
 	int32 BitDepth = InImageWrapper->GetBitDepth();
 	TArray64<uint8> RawData;
 	InImageWrapper->GetRaw(Format, BitDepth, RawData);
@@ -287,12 +290,25 @@ void SImgMediaProcessImages::ProcessImageCustom(TSharedPtr<IImageWrapper>& InIma
 	TArray64<uint8> RawDataTiled;
 	if (bIsTiled)
 	{
-		RawDataTiled.AddUninitialized(RawData.Num());
+		// We don't support tile borders larger than a tile size,
+		// but this shuld not happen in practice.
+		if ((InTileBorder > TileWidth) || (InTileBorder > TileHeight))
+		{
+			UE_LOG(LogImgMediaEditor, Error, TEXT("Tile border is larger than tile size. Clamping to tile size."));
+			InTileBorder = FMath::Min(TileWidth, TileHeight);
+		}
+
+		DestWidth = Width + InTileBorder * 2 * NumTilesX;
+		DestHeight = Height + InTileBorder * 2 * NumTilesY;
+		RawDataTiled.AddUninitialized(DestWidth * DestHeight * BytesPerPixel);
 		RawDataPtr = RawDataTiled.GetData();
 
 		uint8* SourceData = RawData.GetData();
 		uint8* DestData = RawDataTiled.GetData();
+		int32 DestTileWidth = TileWidth + InTileBorder * 2;
+		int32 DestTileHeight = TileHeight + InTileBorder * 2;
 		int32 BytesPerTile = TileWidth * TileHeight * BytesPerPixel;
+		int32 ByterPerDestTile = DestTileWidth * DestTileHeight * BytesPerPixel;
 
 		// Loop over y tiles.
 		for (int32 TileY = 0; TileY < NumTilesY; ++TileY)
@@ -303,14 +319,48 @@ void SImgMediaProcessImages::ProcessImageCustom(TSharedPtr<IImageWrapper>& InIma
 				// Get address of the source and destination tiles.
 				uint8* SourceTile = SourceData +
 					(TileX * TileWidth + TileY * Width * TileHeight) * BytesPerPixel;
-				uint8* DestTile = DestData + (TileX + TileY * NumTilesX) * BytesPerTile;
+				uint8* DestTile = DestData + (TileX + TileY * NumTilesX) * ByterPerDestTile;
+				
+				int32 NumberOfPixelsToCopy = TileWidth;
+
+				// Create a left border.
+				if (TileX > 0)
+				{
+					NumberOfPixelsToCopy += InTileBorder;
+					// Offset the source to get the extra pixels.
+					SourceTile -= InTileBorder * BytesPerPixel;
+				}
+				else
+				{
+					// Offset the destination as we are skipping this border as we have no data.
+					DestTile += InTileBorder * BytesPerPixel;
+				}
+
+				// Create a right border.
+				if (TileX < NumTilesX - 1)
+				{
+					NumberOfPixelsToCopy += InTileBorder;
+				}
 				
 				// Loop over each row in the tile.
-				for (int32 Row = 0; Row < TileHeight; ++Row)
+				for (int32 Row = 0; Row < DestTileHeight; ++Row)
 				{
-					uint8* SourceLine = SourceTile + Row * Width * BytesPerPixel;
-					uint8* DestLine = DestTile + Row * TileWidth * BytesPerPixel;
-					FMemory::Memcpy(DestLine, SourceLine, TileWidth * BytesPerPixel);
+					// Make sure we don't go beyond the source data.
+					int32 SourceRow = Row - InTileBorder;
+					if (TileY == 0)
+					{
+						SourceRow = FMath::Max(SourceRow, 0);
+					}
+					if (TileY == NumTilesY - 1)
+					{
+						SourceRow = FMath::Min(SourceRow, TileHeight - 1);
+					}
+
+					uint8* SourceLine = SourceTile + SourceRow * Width * BytesPerPixel;
+					uint8* DestLine = DestTile + Row * DestTileWidth * BytesPerPixel;
+					
+					// Copy the main data.
+					FMemory::Memcpy(DestLine, SourceLine, NumberOfPixelsToCopy * BytesPerPixel);
 				}
 			}
 		}
@@ -324,11 +374,11 @@ void SImgMediaProcessImages::ProcessImageCustom(TSharedPtr<IImageWrapper>& InIma
 
 	int32 NumChannels = 4;
 
-	FIntPoint Stride(2, Width * BytesPerPixel);
+	FIntPoint Stride(2, DestWidth * BytesPerPixel);
 
 	// Create tiled exr file.
-	FTiledOutputFile OutFile(FIntPoint(0, 0), FIntPoint(Width - 1, Height - 1),
-		FIntPoint(0, 0), FIntPoint(Width - 1, Height - 1));
+	FTiledOutputFile OutFile(FIntPoint(0, 0), FIntPoint(DestWidth - 1, DestHeight - 1),
+		FIntPoint(0, 0), FIntPoint(DestWidth - 1, DestHeight - 1));
 
 	// Add attributes.
 	OutFile.AddIntAttribute(IImgMediaModule::CustomFormatAttributeName.Resolve().ToString(), 1);
@@ -344,11 +394,11 @@ void SImgMediaProcessImages::ProcessImageCustom(TSharedPtr<IImageWrapper>& InIma
 	OutFile.AddChannel(RChannelName);
 
 	// Create output.
-	OutFile.CreateOutputFile(InName, Width, Height, false);
+	OutFile.CreateOutputFile(InName, DestWidth, DestHeight, false);
 	OutFile.AddFrameBufferChannel(AChannelName, RawDataPtr, Stride);
-	OutFile.AddFrameBufferChannel(BChannelName, RawDataPtr + Width * 2, Stride);
-	OutFile.AddFrameBufferChannel(GChannelName, RawDataPtr + Width * 4, Stride);
-	OutFile.AddFrameBufferChannel(RChannelName, RawDataPtr + Width * 6, Stride);
+	OutFile.AddFrameBufferChannel(BChannelName, RawDataPtr + DestWidth * 2, Stride);
+	OutFile.AddFrameBufferChannel(GChannelName, RawDataPtr + DestWidth * 4, Stride);
+	OutFile.AddFrameBufferChannel(RChannelName, RawDataPtr + DestWidth * 6, Stride);
 	OutFile.SetFrameBuffer();
 	
 	OutFile.WriteTile(0, 0, 0);
diff --git a/Engine/Plugins/Media/ImgMedia/Source/ImgMediaEditor/Private/Widgets/SImgMediaProcessImages.h b/Engine/Plugins/Media/ImgMedia/Source/ImgMediaEditor/Private/Widgets/SImgMediaProcessImages.h
--- a/Engine/Plugins/Media/ImgMedia/Source/ImgMediaEditor/Private/Widgets/SImgMediaProcessImages.h
+++ b/Engine/Plugins/Media/ImgMedia/Source/ImgMediaEditor/Private/Widgets/SImgMediaProcessImages.h
@@ -58,10 +58,11 @@ private:
 	 * @param InImageWrapper	ImageWrapper to read/write the image.
 	 * @param InTileWidth		Desired width of tiles.
 	 * @param InTileHeight		Desired height of tiles.
+	 * @param InTileBorder		Number of pixels to duplicate along a tile edge.
 	 * @param InName			Full path and name of file to write.
 	 */
 	void ProcessImageCustom(TSharedPtr<IImageWrapper>& InImageWrapper,
-		int32 InTileWidth, int32 InTileHeight, const FString& InName);
+		int32 InTileWidth, int32 InTileHeight, int32 InTileBorder, const FString& InName);
 
 	/** Holds our details view. */
 	TSharedPtr<class IDetailsView> DetailsView;
