ImgMedia Exr Reader and Media Process: Reading and writing custom attribute for tile border, as well as reading and storing other attributes.

#rnx

diff --git a/Engine/Plugins/Media/ImgMedia/Source/ImgMedia/Private/ImgMediaModule.cpp b/Engine/Plugins/Media/ImgMedia/Source/ImgMedia/Private/ImgMediaModule.cpp
--- a/Engine/Plugins/Media/ImgMedia/Source/ImgMedia/Private/ImgMediaModule.cpp
+++ b/Engine/Plugins/Media/ImgMedia/Source/ImgMedia/Private/ImgMediaModule.cpp
@@ -18,6 +18,7 @@ DEFINE_LOG_CATEGORY(LogImgMedia);
 FLazyName IImgMediaModule::CustomFormatAttributeName(TEXT("EpicGamesCustomFormat"));
 FLazyName IImgMediaModule::CustomFormatTileWidthAttributeName(TEXT("EpicGamesTileWidth"));
 FLazyName IImgMediaModule::CustomFormatTileHeightAttributeName(TEXT("EpicGamesTileHeight"));
+FLazyName IImgMediaModule::CustomFormatTileBorderAttributeName(TEXT("EpicGamesTileBorder"));
 
 TSharedPtr<FImgMediaGlobalCache, ESPMode::ThreadSafe> IImgMediaModule::GlobalCache;
 
diff --git a/Engine/Plugins/Media/ImgMedia/Source/ImgMedia/Private/Readers/ExrImgMediaReader.cpp b/Engine/Plugins/Media/ImgMedia/Source/ImgMedia/Private/Readers/ExrImgMediaReader.cpp
--- a/Engine/Plugins/Media/ImgMedia/Source/ImgMedia/Private/Readers/ExrImgMediaReader.cpp
+++ b/Engine/Plugins/Media/ImgMedia/Source/ImgMedia/Private/Readers/ExrImgMediaReader.cpp
@@ -297,12 +297,27 @@ bool FExrImgMediaReader::GetInfo(FRgbaInputFile& InputFile, FImgMediaFrameInfo&
 	OutInfo.Srgb = false;
 	OutInfo.UncompressedSize = InputFile.GetUncompressedSize();
 	OutInfo.NumChannels = InputFile.GetNumChannels();
-	OutInfo.bHasTiles = InputFile.GetTileSize(OutInfo.TileDimensions);
+
+	int32 CustomFormat = 0;
+	InputFile.GetIntAttribute(IImgMediaModule::CustomFormatAttributeName.Resolve().ToString(), CustomFormat);
+	bool bIsCustomFormat = CustomFormat > 0;
+	if (bIsCustomFormat)
+	{
+		// Get tile size.
+		OutInfo.bHasTiles = InputFile.GetIntAttribute(IImgMediaModule::CustomFormatTileWidthAttributeName.Resolve().ToString(), OutInfo.TileDimensions.X);
+		OutInfo.bHasTiles = OutInfo.bHasTiles && InputFile.GetIntAttribute(IImgMediaModule::CustomFormatTileHeightAttributeName.Resolve().ToString(), OutInfo.TileDimensions.Y);
+		InputFile.GetIntAttribute(IImgMediaModule::CustomFormatTileBorderAttributeName.Resolve().ToString(), OutInfo.TileBorder);
+	}
+	else
+	{
+		OutInfo.bHasTiles = InputFile.GetTileSize(OutInfo.TileDimensions);
+		OutInfo.TileBorder = 0;
+	}
+
 	if (OutInfo.bHasTiles)
 	{
 		OutInfo.NumTiles = FIntPoint(OutInfo.Dim.X / OutInfo.TileDimensions.X, OutInfo.Dim.Y / OutInfo.TileDimensions.Y);
 	}
-
 	return (OutInfo.UncompressedSize > 0) && (OutInfo.Dim.GetMin() > 0);
 }
 
diff --git a/Engine/Plugins/Media/ImgMedia/Source/ImgMedia/Private/Readers/IImgMediaReader.h b/Engine/Plugins/Media/ImgMedia/Source/ImgMedia/Private/Readers/IImgMediaReader.h
--- a/Engine/Plugins/Media/ImgMedia/Source/ImgMedia/Private/Readers/IImgMediaReader.h
+++ b/Engine/Plugins/Media/ImgMedia/Source/ImgMedia/Private/Readers/IImgMediaReader.h
@@ -46,6 +46,9 @@ struct FImgMediaFrameInfo
 
 	/** Tile dimensions. */
 	FIntPoint TileDimensions;
+
+	/** Tile border size in texels. This is required for more elaborate tile texel sampling. */
+	int32 TileBorder;
 };
 
 
diff --git a/Engine/Plugins/Media/ImgMedia/Source/ImgMedia/Public/IImgMediaModule.h b/Engine/Plugins/Media/ImgMedia/Source/ImgMedia/Public/IImgMediaModule.h
--- a/Engine/Plugins/Media/ImgMedia/Source/ImgMedia/Public/IImgMediaModule.h
+++ b/Engine/Plugins/Media/ImgMedia/Source/ImgMedia/Public/IImgMediaModule.h
@@ -45,6 +45,8 @@ public:
 	static FLazyName CustomFormatTileWidthAttributeName;
 	/** Name of attribute in the Exr file for the tile height for our custom format. */
 	static FLazyName CustomFormatTileHeightAttributeName;
+	/** Name of attribute in the Exr file for the tile border size for our custom format. */
+	static FLazyName CustomFormatTileBorderAttributeName;
 
 protected:
 
diff --git a/Engine/Plugins/Media/ImgMedia/Source/ImgMediaEditor/Private/Widgets/SImgMediaProcessImages.cpp b/Engine/Plugins/Media/ImgMedia/Source/ImgMediaEditor/Private/Widgets/SImgMediaProcessImages.cpp
--- a/Engine/Plugins/Media/ImgMedia/Source/ImgMediaEditor/Private/Widgets/SImgMediaProcessImages.cpp
+++ b/Engine/Plugins/Media/ImgMedia/Source/ImgMediaEditor/Private/Widgets/SImgMediaProcessImages.cpp
@@ -382,10 +382,14 @@ void SImgMediaProcessImages::ProcessImageCustom(TSharedPtr<IImageWrapper>& InIma
 
 	// Add attributes.
 	OutFile.AddIntAttribute(IImgMediaModule::CustomFormatAttributeName.Resolve().ToString(), 1);
-	OutFile.AddIntAttribute(IImgMediaModule::CustomFormatTileWidthAttributeName.Resolve().ToString(),
-		bIsTiled ? TileWidth : 0);
-	OutFile.AddIntAttribute(IImgMediaModule::CustomFormatTileHeightAttributeName.Resolve().ToString(),
-		bIsTiled ? TileHeight : 0);
+
+	// These attributes will not be added and therefore not found by EXR reader if it is not tiled.
+	if (bIsTiled)
+	{
+		OutFile.AddIntAttribute(IImgMediaModule::CustomFormatTileWidthAttributeName.Resolve().ToString(), TileWidth);
+		OutFile.AddIntAttribute(IImgMediaModule::CustomFormatTileHeightAttributeName.Resolve().ToString(),TileHeight);
+		OutFile.AddIntAttribute(IImgMediaModule::CustomFormatTileBorderAttributeName.Resolve().ToString(), InTileBorder);
+	}
 
 	// Add channels.
 	OutFile.AddChannel(AChannelName);
