ImgMedia: Adds explicit support for sets of tiles per mip level, enabling visible tiles from virtual texturing.

diff --git a/Engine/Plugins/Media/ImgMedia/Source/ImgMedia/Private/Loader/ImgMediaLoader.cpp b/Engine/Plugins/Media/ImgMedia/Source/ImgMedia/Private/Loader/ImgMediaLoader.cpp
--- a/Engine/Plugins/Media/ImgMedia/Source/ImgMedia/Private/Loader/ImgMediaLoader.cpp
+++ b/Engine/Plugins/Media/ImgMedia/Source/ImgMedia/Private/Loader/ImgMediaLoader.cpp
@@ -50,6 +50,69 @@ namespace ImgMediaLoader
 		}
 	}
 	
+
+	FImgMediaTextureSample::TilingDescription GetTilingDescription(const FImgMediaLoader* Loader)
+	{
+		FImgMediaTextureSample::TilingDescription TilingDesc;
+		
+		if (Loader->IsTiled())
+		{
+			const int32 NumX = Loader->GetNumTilesX();
+			const int32 NumY = Loader->GetNumTilesY();
+
+			TilingDesc.TileBorderSize = 0; // todo: get from loader once accessible from metadata
+			TilingDesc.TileSize = FIntPoint(Loader->GetSequenceDim().X / NumX, Loader->GetSequenceDim().Y / NumY);
+			TilingDesc.TileNum = FIntPoint(NumX, NumY);
+		}
+		
+		return TilingDesc;
+	}
+
+	// Check if the existing tiles contain all of the requested ones. (Is existing a superset of requested?)
+	bool ContainsMipTiles(const TMap<int32, FImgMediaTileSelection>& ExistingTiles, const TMap<int32, FImgMediaTileSelection>& RequestedTiles)
+	{
+		for (auto Iter = RequestedTiles.CreateConstIterator(); Iter; ++Iter)
+		{
+			const int32 RequestedMipLevel = Iter.Key();
+			const FImgMediaTileSelection& RequestedSelection = Iter.Value();
+
+			if (const FImgMediaTileSelection* ExistingSelection = ExistingTiles.Find(RequestedMipLevel))
+			{
+				if (!ExistingSelection->Contains(RequestedSelection))
+				{
+					// Requested tile selection is not present.
+					return false;
+				}
+			}
+			else
+			{
+				// Requested mip level is not present.
+				return false;
+			}
+		}
+		
+		// Requested tiles already exist, or the request was empty.
+		return true;
+	}
+
+	// Reduce the selected tiles region for higher mip levels, always ensuring that all tiles at the base level are included.
+	FImgMediaTileSelection DownscaleTileSelection(const FImgMediaTileSelection& InTopLevelSelection, int32 InMipLevel)
+	{
+		if (InMipLevel <= 0)
+		{
+			return InTopLevelSelection;
+		}
+		
+		const int32 MipLevelDiv = 1 << InMipLevel;
+
+		FImgMediaTileSelection OutSelection;
+		OutSelection.TopLeftX = (uint16)(InTopLevelSelection.TopLeftX / MipLevelDiv);
+		OutSelection.TopLeftY = (uint16)(InTopLevelSelection.TopLeftY / MipLevelDiv);
+		OutSelection.BottomRightX = (uint16)FMath::CeilToInt(float(InTopLevelSelection.BottomRightX) / MipLevelDiv);
+		OutSelection.BottomRightY = (uint16)FMath::CeilToInt(float(InTopLevelSelection.BottomRightY) / MipLevelDiv);
+
+		return OutSelection;
+	}
 }
 
 /* FImgMediaLoader structors
@@ -174,7 +237,7 @@ TSharedPtr<FImgMediaTextureSample, ESPMode::ThreadSafe> FImgMediaLoader::GetFram
 
 	ImgMediaLoader::CheckAndUpdateImgDimensions(SequenceDim, Frame->Get()->Info.Dim);
 
-	if (!Sample->Initialize(*Frame->Get(), SequenceDim, FMediaTimeStamp(FrameStartTime, 0), NextStartTime - FrameStartTime, GetNumMipLevels()))
+	if (!Sample->Initialize(*Frame->Get(), SequenceDim, FMediaTimeStamp(FrameStartTime, 0), NextStartTime - FrameStartTime, GetNumMipLevels(), ImgMediaLoader::GetTilingDescription(this)))
 	{
 		return nullptr;
 	}
@@ -475,7 +538,7 @@ IMediaSamples::EFetchBestSampleResult FImgMediaLoader::FetchBestVideoSampleForTi
 
 					ImgMediaLoader::CheckAndUpdateImgDimensions(SequenceDim, Frame->Get()->Info.Dim);
 
-					if (Sample->Initialize(*Frame->Get(), SequenceDim, FMediaTimeStamp(FrameNumberToTime(MaxIdx), QueuedSampleFetch.CurrentSequenceIndex), FTimespan::FromSeconds(Duration), GetNumMipLevels()))
+					if (Sample->Initialize(*Frame->Get(), SequenceDim, FMediaTimeStamp(FrameNumberToTime(MaxIdx), QueuedSampleFetch.CurrentSequenceIndex), FTimespan::FromSeconds(Duration), GetNumMipLevels(), ImgMediaLoader::GetTilingDescription(this)))
 					{
 						OutSample = Sample;
 						return IMediaSamples::EFetchBestSampleResult::Ok;
@@ -619,11 +682,11 @@ IQueuedWork* FImgMediaLoader::GetWork()
 		ExistingFrame = *ExistingFramePtr;
 	}
 
-	FImgMediaTileSelection TileSelection;
-	int MipLevel = GetDesiredMipLevel(FrameNumber, TileSelection);
+	TMap<int32, FImgMediaTileSelection> DesiredMipsAndTiles;
+	GetDesiredMipTiles(FrameNumber, DesiredMipsAndTiles);
 	
 	// Set up work.
-	Work->Initialize(FrameNumber, MipLevel, TileSelection, ExistingFrame);
+	Work->Initialize(FrameNumber, DesiredMipsAndTiles, ExistingFrame);
 	QueuedFrameNumbers.Add(FrameNumber);
 
 	return Work;
@@ -800,6 +863,15 @@ void FImgMediaLoader::LoadSequence(const FString& SequencePath, const FFrameRate
 	SequenceDuration = FrameNumberToTime(GetNumImages());
 	SIZE_T UncompressedSize = FirstFrameInfo.UncompressedSize;
 
+#if 0
+	// todo: enable once both the CPU & GPU reader use the same tile count logic.
+	if (FirstFrameInfo.bHasTiles)
+	{
+		NumTilesX = FirstFrameInfo.NumTiles.X;
+		NumTilesY = FirstFrameInfo.NumTiles.Y;
+	}
+#endif
+
 	// If we have no mips or tiles, then get rid of our MipMapInfoObject.
 	// Otherwise, set it up.
 	if (MipMapInfo.IsValid())
@@ -1188,10 +1260,11 @@ void FImgMediaLoader::Update(int32 PlayHeadFrame, float PlayRate, bool Loop)
 		}
 		else
 		{
-			// Yes. Check if we have the desired mip level.
-			FImgMediaTileSelection TileSelection;
-			int32 MipLevel = GetDesiredMipLevel(FrameNumber, TileSelection);
-			NeedFrame = ((*FramePtr)->MipMapsPresent & (1 << MipLevel)) == 0;
+			// Yes. Check if we have all the desired tiles per mip level.
+			TMap<int32, FImgMediaTileSelection> DesiredMipsAndTiles;
+			GetDesiredMipTiles(FrameNumber, DesiredMipsAndTiles);
+
+			NeedFrame = !ImgMediaLoader::ContainsMipTiles((*FramePtr)->MipTilesPresent, DesiredMipsAndTiles);
 		}
 		
 		if ((NeedFrame) && !QueuedFrameNumbers.Contains(FrameNumber))
@@ -1226,9 +1299,13 @@ void FImgMediaLoader::AddEmptyFrame(int32 FrameNumber)
 	Frame->Info.Dim = SequenceDim;
 	Frame->Info.FrameRate = SequenceFrameRate;
 	Frame->Info.NumChannels = NumChannels;
+	Frame->Info.bHasTiles = false;
 	Frame->Format = EMediaTextureSampleFormat::FloatRGB;
 	Frame->Stride = Frame ->Info.Dim.X * PixelSize;
-	Frame->MipMapsPresent = -1;
+	for (int32 Level = 0; Level < GetNumMipLevels(); ++Level)
+	{
+		Frame->MipTilesPresent.Add(Level).SetAllVisible();
+	}
 	AddFrameToCache(FrameNumber, Frame);
 }
 
@@ -1266,16 +1343,39 @@ void FImgMediaLoader::AddFrameToCache(int32 FrameNumber, const TSharedPtr<FImgMe
 	}
 }
 
-int32 FImgMediaLoader::GetDesiredMipLevel(int32 FrameIndex, FImgMediaTileSelection& OutTileSelection)
+void FImgMediaLoader::GetDesiredMipTiles(int32 FrameIndex, TMap<int32, FImgMediaTileSelection>& OutMipsAndTiles)
 {
-	int32 MipLevel = 0;
-
-	if (MipMapInfo.IsValid())
+	if(GetDefault<UImgMediaSettings>()->bUseVirtualTextureTiles)
 	{
-		MipLevel = MipMapInfo->GetDesiredMipLevel(OutTileSelection);
+		TMap<int32, TSet<FMediaTileCoordinate>> VisibleTiles;
+		GetVisibleTiles(VisibleTiles);
+
+		for (const auto& Pair : VisibleTiles)
+		{
+			FImgMediaTileSelection& OutSelection = OutMipsAndTiles.Add(Pair.Key);
+			OutSelection.SetAllNotVisible();
+
+			for (const FMediaTileCoordinate& Coordinate : Pair.Value)
+			{
+				OutSelection.Include(Coordinate.X, Coordinate.Y);
+			}
+		}
 	}
+	else
+	{
+		FImgMediaTileSelection BaseTileSelection;
+		BaseTileSelection.SetAllVisible();
 
-	return MipLevel;
+		const int32 MipLevel = MipMapInfo.IsValid() ? MipMapInfo->GetDesiredMipLevel(BaseTileSelection) : 0;
+
+		OutMipsAndTiles.Emplace(MipLevel, ImgMediaLoader::DownscaleTileSelection(BaseTileSelection, MipLevel));
+
+		// Reproducing the previous mip-threshold logic by adding all higher mip levels.
+		for (int32 HigherMipLevel = MipLevel + 1; HigherMipLevel < GetNumMipLevels(); ++HigherMipLevel)
+		{
+			OutMipsAndTiles.Emplace(HigherMipLevel, ImgMediaLoader::DownscaleTileSelection(BaseTileSelection, HigherMipLevel));
+		}
+	}
 }
 
 
@@ -1360,3 +1460,19 @@ bool FImgMediaLoader::GetNumberAtEndOfString(int32 &Number, const FString& Strin
 	
 	return bFoundNumber;
 }
+
+bool FImgMediaLoader::GetVisibleTiles(TMap<int32, TSet<FMediaTileCoordinate>>& OutTiles) const
+{
+	FScopeLock Lock(&CriticalSectionTiles);
+	OutTiles = ActiveTilesPerMipLevel;
+
+	return true;
+}
+
+bool FImgMediaLoader::SetVisibleTiles(TMap<int32, TSet<FMediaTileCoordinate>>&& InTiles)
+{
+	FScopeLock Lock(&CriticalSectionTiles);
+	ActiveTilesPerMipLevel = MoveTemp(InTiles);
+
+	return true;
+}
diff --git a/Engine/Plugins/Media/ImgMedia/Source/ImgMedia/Private/Loader/ImgMediaLoader.h b/Engine/Plugins/Media/ImgMedia/Source/ImgMedia/Private/Loader/ImgMediaLoader.h
--- a/Engine/Plugins/Media/ImgMedia/Source/ImgMedia/Private/Loader/ImgMediaLoader.h
+++ b/Engine/Plugins/Media/ImgMedia/Source/ImgMedia/Private/Loader/ImgMediaLoader.h
@@ -6,6 +6,8 @@
 #include "Containers/LruCache.h"
 #include "Containers/Queue.h"
 #include "IMediaSamples.h"
+#include "IMediaView.h"
+#include "ImgMediaMipMapInfo.h"
 #include "Misc/FrameRate.h"
 #include "Templates/SharedPointer.h"
 
@@ -20,7 +22,6 @@ class IImgMediaReader;
 class IQueuedWork;
 
 struct FImgMediaFrame;
-struct FImgMediaTileSelection;
 
 
 /**
@@ -375,13 +376,12 @@ protected:
 	void AddFrameToCache(int32 FrameNumber, const TSharedPtr<FImgMediaFrame, ESPMode::ThreadSafe>& Frame);
 
 	/**
-	 * Get what mip level we should be using for a given frame.
+	 * Get what mip level and tiles we should be using for a given frame.
 	 *
 	 * @param FrameIndex Frame to get mip level for.
 	 * @param OutTileSelection Will be filled in with what tiles are actually needed.
-	 * @return Returns the mip level.
 	 */
-	int32 GetDesiredMipLevel(int32 FrameIndex, FImgMediaTileSelection& OutTileSelection);
+	void GetDesiredMipTiles(int32 FrameIndex, TMap<int32, FImgMediaTileSelection>& OutMipsAndTiles);
 
 	/***
 	 * Modulos the time so that it is between 0 and SequenceDuration.
@@ -422,6 +422,14 @@ protected:
 	 */
 	bool GetNumberAtEndOfString(int32& Number, const FString& String) const;
 
+public:
+	
+	/** Thread-safe function to get currently visible tiles per mip level. */
+	bool GetVisibleTiles(TMap<int32, TSet<FMediaTileCoordinate>>& OutTiles) const;
+
+	/** Thread-safe function to set currently visible tiles per mip level. */
+	bool SetVisibleTiles(TMap<int32, TSet<FMediaTileCoordinate>>&& InTiles);
+
 private:
 
 	/** Critical section for synchronizing access to Frames. */
@@ -516,4 +524,10 @@ private:
 		uint64 CurrentSequenceIndex;
 	} QueuedSampleFetch;
 
+private:
+	/** Critical section for synchronizing access to the active tile list. */
+	mutable FCriticalSection CriticalSectionTiles;
+
+	/** List of visible tiles updated through the IMediaView interface. */
+	TMap<int32, TSet<FMediaTileCoordinate>> ActiveTilesPerMipLevel;
 };
diff --git a/Engine/Plugins/Media/ImgMedia/Source/ImgMedia/Private/Loader/ImgMediaLoaderWork.cpp b/Engine/Plugins/Media/ImgMedia/Source/ImgMedia/Private/Loader/ImgMediaLoaderWork.cpp
--- a/Engine/Plugins/Media/ImgMedia/Source/ImgMedia/Private/Loader/ImgMediaLoaderWork.cpp
+++ b/Engine/Plugins/Media/ImgMedia/Source/ImgMedia/Private/Loader/ImgMediaLoaderWork.cpp
@@ -33,12 +33,11 @@ FImgMediaLoaderWork::FImgMediaLoaderWork(const TSharedRef<FImgMediaLoader, ESPMo
  *****************************************************************************/
 
 void FImgMediaLoaderWork::Initialize(int32 InFrameNumber,
-	int32 InMipLevel, const FImgMediaTileSelection& InTileSelection,
+	TMap<int32, FImgMediaTileSelection> InMipTiles,
 	TSharedPtr<FImgMediaFrame, ESPMode::ThreadSafe> InExistingFrame)
 {
 	FrameNumber = InFrameNumber;
-	MipLevel = InMipLevel;
-	TileSelection = InTileSelection;
+	MipTiles = MoveTemp(InMipTiles);
 	ExistingFrame = InExistingFrame;
 }
 
@@ -77,7 +76,7 @@ void FImgMediaLoaderWork::DoThreadedWork()
 			Frame = ExistingFrame;
 		}
 
-		if (!Reader->ReadFrame(FrameNumber, MipLevel, TileSelection, Frame))
+		if(!Reader->ReadFrame(FrameNumber, MipTiles, Frame))
 		{
 			if (ExistingFrame == nullptr)
 			{
diff --git a/Engine/Plugins/Media/ImgMedia/Source/ImgMedia/Private/Loader/ImgMediaLoaderWork.h b/Engine/Plugins/Media/ImgMedia/Source/ImgMedia/Private/Loader/ImgMediaLoaderWork.h
--- a/Engine/Plugins/Media/ImgMedia/Source/ImgMedia/Private/Loader/ImgMediaLoaderWork.h
+++ b/Engine/Plugins/Media/ImgMedia/Source/ImgMedia/Private/Loader/ImgMediaLoaderWork.h
@@ -39,7 +39,7 @@ public:
 	 * @param InExistingFrame If set, then use this frame to load into.
 	 * @see Shutdown
 	 */
-	void Initialize(int32 InFrameNumber, int32 InMipLevel, const FImgMediaTileSelection& InTileSelection, TSharedPtr<FImgMediaFrame, ESPMode::ThreadSafe> InExistingFrame);
+	void Initialize(int32 InFrameNumber, TMap<int32, FImgMediaTileSelection> InMipTiles, TSharedPtr<FImgMediaFrame, ESPMode::ThreadSafe> InExistingFrame);
 
 public:
 
@@ -62,11 +62,8 @@ private:
 	/** The number of the image frame. */
 	int32 FrameNumber;
 
-	/** Mip level of the image frame. */
-	int32 MipLevel;
-
-	/** Which tiles we should load. */
-	FImgMediaTileSelection TileSelection;
+	/** Which tiles we should load per mip level. */
+	TMap<int32, FImgMediaTileSelection> MipTiles;
 	
 	/** If set, then load the image into this frame instead of making a new one. */
 	TSharedPtr<FImgMediaFrame, ESPMode::ThreadSafe> ExistingFrame;
diff --git a/Engine/Plugins/Media/ImgMedia/Source/ImgMedia/Private/Player/ImgMediaPlayer.cpp b/Engine/Plugins/Media/ImgMedia/Source/ImgMedia/Private/Player/ImgMediaPlayer.cpp
--- a/Engine/Plugins/Media/ImgMedia/Source/ImgMedia/Private/Player/ImgMediaPlayer.cpp
+++ b/Engine/Plugins/Media/ImgMedia/Source/ImgMedia/Private/Player/ImgMediaPlayer.cpp
@@ -407,6 +407,29 @@ bool FImgMediaPlayer::GetPlayerFeatureFlag(EFeatureFlag flag) const
 	return IMediaPlayer::GetPlayerFeatureFlag(flag);
 }
 
+/* IMediaView interface
+ *****************************************************************************/
+
+bool FImgMediaPlayer::GetVisibleTiles(TMap<int32, TSet<FMediaTileCoordinate>>& OutTiles) const
+{
+	if (Loader.IsValid())
+	{
+		return Loader->GetVisibleTiles(OutTiles);
+	}
+
+	return false;
+}
+
+bool FImgMediaPlayer::SetVisibleTiles(TMap<int32, TSet<FMediaTileCoordinate>>&& InTiles)
+{
+	if (Loader.IsValid())
+	{
+		return Loader->SetVisibleTiles(MoveTemp(InTiles));
+	}
+
+	return false;
+}
+
 
 /* FImgMediaPlayer implementation
  *****************************************************************************/
diff --git a/Engine/Plugins/Media/ImgMedia/Source/ImgMedia/Private/Player/ImgMediaPlayer.h b/Engine/Plugins/Media/ImgMedia/Source/ImgMedia/Private/Player/ImgMediaPlayer.h
--- a/Engine/Plugins/Media/ImgMedia/Source/ImgMedia/Private/Player/ImgMediaPlayer.h
+++ b/Engine/Plugins/Media/ImgMedia/Source/ImgMedia/Private/Player/ImgMediaPlayer.h
@@ -27,7 +27,7 @@ class FImgMediaPlayer
 	, protected IMediaControls
 	, protected IMediaSamples
 	, protected IMediaTracks
-	, protected IMediaView
+	, public IMediaView
 {
 public:
 
@@ -65,6 +65,13 @@ public:
 	virtual void ProcessVideoSamples() override;
 	virtual bool GetPlayerFeatureFlag(EFeatureFlag flag) const override;
 
+public:
+
+	//~ IMediaView interface
+
+	virtual bool GetVisibleTiles(TMap<int32, TSet<FMediaTileCoordinate>>& OutTiles) const override;
+	virtual bool SetVisibleTiles(TMap<int32, TSet<FMediaTileCoordinate>>&& InTiles) override;
+
 protected:
 
 	/**
diff --git a/Engine/Plugins/Media/ImgMedia/Source/ImgMedia/Private/Player/ImgMediaTextureSample.h b/Engine/Plugins/Media/ImgMedia/Source/ImgMedia/Private/Player/ImgMediaTextureSample.h
--- a/Engine/Plugins/Media/ImgMedia/Source/ImgMedia/Private/Player/ImgMediaTextureSample.h
+++ b/Engine/Plugins/Media/ImgMedia/Source/ImgMedia/Private/Player/ImgMediaTextureSample.h
@@ -24,6 +24,7 @@ public:
 		: Duration(FTimespan::Zero())
 		, OutputDim(FIntPoint::ZeroValue)
 		, Time(FTimespan::Zero())
+		, TilingDesc()
 	{ }
 
 public:
@@ -36,13 +37,14 @@ public:
 	 * @param InTime The sample time (in the player's local clock).
 	 * @param InDuration The duration for which the sample is valid.
 	 */
-	bool Initialize(FImgMediaFrame& InFrame, const FIntPoint& InOutputDim, FMediaTimeStamp InTime, FTimespan InDuration, uint8 InNumMipMaps)
+	bool Initialize(FImgMediaFrame& InFrame, const FIntPoint& InOutputDim, FMediaTimeStamp InTime, FTimespan InDuration, uint8 InNumMipMaps, TilingDescription InTilingDesc = TilingDescription())
 	{
 		Duration = InDuration;
 		Frame = InFrame;
 		OutputDim = InOutputDim;
 		Time = InTime;
 		NumMipMaps = InNumMipMaps;
+		TilingDesc = InTilingDesc;
 
 		// If we have no data then make sure the number of mipmaps is 1.
 		// otherwise FMediaTextureResource won't like it.
@@ -73,6 +75,11 @@ public:
 		return NumMipMaps;
 	}
 
+	virtual TilingDescription GetTilingDescription() const override
+	{
+		return TilingDesc;
+	}
+
 	virtual FTimespan GetDuration() const override
 	{
 		return Duration;
@@ -136,4 +143,7 @@ private:
 
 	/** Number of mip levels in this sample. */
 	uint8 NumMipMaps;
+
+	/** Description of the number and size of tiles in this sample. */
+	TilingDescription TilingDesc;
 };
diff --git a/Engine/Plugins/Media/ImgMedia/Source/ImgMedia/Private/Readers/ExrImgMediaReader.cpp b/Engine/Plugins/Media/ImgMedia/Source/ImgMedia/Private/Readers/ExrImgMediaReader.cpp
--- a/Engine/Plugins/Media/ImgMedia/Source/ImgMedia/Private/Readers/ExrImgMediaReader.cpp
+++ b/Engine/Plugins/Media/ImgMedia/Source/ImgMedia/Private/Readers/ExrImgMediaReader.cpp
@@ -73,7 +73,7 @@ bool FExrImgMediaReader::GetFrameInfo(const FString& ImagePath, FImgMediaFrameIn
 }
 
 
-bool FExrImgMediaReader::ReadFrame(int32 FrameId, int32 MipLevel, const FImgMediaTileSelection& InTileSelection, TSharedPtr<FImgMediaFrame, ESPMode::ThreadSafe> OutFrame)
+bool FExrImgMediaReader::ReadFrame(int32 FrameId, const TMap<int32, FImgMediaTileSelection>& InMipTiles, TSharedPtr<FImgMediaFrame, ESPMode::ThreadSafe> OutFrame)
 {
 	TSharedPtr<FImgMediaLoader, ESPMode::ThreadSafe> Loader = LoaderPtr.Pin();
 	if (Loader.IsValid() == false)
@@ -85,10 +85,6 @@ bool FExrImgMediaReader::ReadFrame(int32 FrameId, int32 MipLevel, const FImgMedi
 	int32 NumTilesX = Loader->GetNumTilesX();
 	int32 NumTilesY = Loader->GetNumTilesY();
 	bool bHasTiles = (NumTilesX * NumTilesY) > 1;
-	int32 StartTileX = InTileSelection.TopLeftX;
-	int32 StartTileY = InTileSelection.TopLeftY;
-	int32 EndTileX = FMath::Min((int32)InTileSelection.BottomRightX, NumTilesX);
-	int32 EndTileY = FMath::Min((int32)InTileSelection.BottomRightY, NumTilesY);
 
 	int32 BytesPerPixelPerChannel = sizeof(uint16);
 	int32 NumChannels = 4;
@@ -145,80 +141,92 @@ bool FExrImgMediaReader::ReadFrame(int32 FrameId, int32 MipLevel, const FImgMedi
 		// The EXR RGBA interface only outputs RGBA data.
 		OutFrame->Format = EMediaTextureSampleFormat::FloatRGBA;
 		OutFrame->Data = MakeShareable(Buffer, MoveTemp(BufferDeleter));
-		OutFrame->MipMapsPresent = 0;
+		OutFrame->MipTilesPresent.Reset();
 		OutFrame->Stride = OutFrame->Info.Dim.X * BytesPerPixel;
 	}
 
 	// Loop over all mips.
 	uint8* MipDataPtr = (uint8*)(OutFrame->Data.Get());
 	FIntPoint Dim = OutFrame->Info.Dim;
-	bool LevelFoundSoFar = false;
+	
 	int32 NumMipLevels = Loader->GetNumMipLevels();
 	for (int32 CurrentMipLevel = 0; CurrentMipLevel < NumMipLevels; ++CurrentMipLevel)
 	{
-		// Do we want to read in this mip?
-		bool IsThisLevelPresent = (OutFrame->MipMapsPresent & (1 << CurrentMipLevel)) != 0;
-		bool ReadThisMip = (CurrentMipLevel >= MipLevel) &&
-			(IsThisLevelPresent == false);
-		
-		if (ReadThisMip)
+		if (InMipTiles.Contains(CurrentMipLevel))
 		{
-			FString Image = Loader->GetImagePath(FrameId, CurrentMipLevel);
-			FString BaseImage;
-			if (bHasTiles)
+			const FImgMediaTileSelection& CurrentTileSelection = InMipTiles[CurrentMipLevel];
+
+			const int MipLevelDiv = 1 << CurrentMipLevel;
+			int32 StartTileX = CurrentTileSelection.TopLeftX;
+			int32 StartTileY = CurrentTileSelection.TopLeftY;
+			int32 EndTileX = FMath::Min((int32)CurrentTileSelection.BottomRightX, FMath::CeilToInt(float(NumTilesX) / MipLevelDiv));
+			int32 EndTileY = FMath::Min((int32)CurrentTileSelection.BottomRightY, FMath::CeilToInt(float(NumTilesY) / MipLevelDiv));
+
+			bool ReadThisMip = true;
+			// Avoid reads if the cached frame already contains the current tiles for this mip level.
+			if (const FImgMediaTileSelection* CachedSelection = OutFrame->MipTilesPresent.Find(CurrentMipLevel))
 			{
-				// Remove "_x0_y0.exr" so we can add on the correct name for the tile we want.
-				BaseImage = Image.LeftChop(10);
+				ReadThisMip = !CachedSelection->Contains(CurrentTileSelection);
 			}
 
-			int32 TileWidth = Dim.X / NumTilesX;
-			int32 TileHeight = Dim.Y / NumTilesY;
-			// The offset into the frame buffer for each row of tiles.
-			// Total width * height of a tile * bytes per pixel.
-			int32 FrameBufferOffsetPerTileY = Dim.X * TileHeight * BytesPerPixel;
-			int32 FrameBufferOffsetY = FrameBufferOffsetPerTileY * StartTileY;
-			
-			for (int32 TileY = StartTileY; TileY < EndTileY; TileY++)
+			if (ReadThisMip)
 			{
-				// The offset into the frame buffer for each column of tiles.
-				// Tile width * bytes per pixel.
-				int32 FrameBufferOffsetPerTileX = TileWidth * BytesPerPixel;
-				int32 FrameBufferOffsetX = FrameBufferOffsetPerTileX * StartTileX;
-				for (int32 TileX = StartTileX; TileX < EndTileX; TileX++)
+				FString Image = Loader->GetImagePath(FrameId, CurrentMipLevel);
+				FString BaseImage;
+				if (bHasTiles)
 				{
-					// Get for our frame/mip level.
-					if (bHasTiles)
-					{
-						Image = FString::Printf(TEXT("%s_x%d_y%d.exr"), *BaseImage, TileX, TileY);
-					}
-					FRgbaInputFile InputFile(Image, 2);
-					if (InputFile.HasInputFile())
-					{
-						// read frame data
-						InputFile.SetFrameBuffer(MipDataPtr + FrameBufferOffsetX + FrameBufferOffsetY, Dim);
-						InputFile.ReadPixels(0, TileHeight - 1);
+					// Remove "_x0_y0.exr" so we can add on the correct name for the tile we want.
+					BaseImage = Image.LeftChop(10);
+				}
 
-						FrameBufferOffsetX += FrameBufferOffsetPerTileX;
-						OutFrame->MipMapsPresent |= 1 << CurrentMipLevel;
-						LevelFoundSoFar = true;
-					}
-					else
+				int32 TileWidth = Dim.X / NumTilesX;
+				int32 TileHeight = Dim.Y / NumTilesY;
+				// The offset into the frame buffer for each row of tiles.
+				// Total width * height of a tile * bytes per pixel.
+				int32 FrameBufferOffsetPerTileY = Dim.X * TileHeight * BytesPerPixel;
+				int32 FrameBufferOffsetY = FrameBufferOffsetPerTileY * StartTileY;
+
+				for (int32 TileY = StartTileY; TileY < EndTileY; TileY++)
+				{
+					// The offset into the frame buffer for each column of tiles.
+					// Tile width * bytes per pixel.
+					int32 FrameBufferOffsetPerTileX = TileWidth * BytesPerPixel;
+					int32 FrameBufferOffsetX = FrameBufferOffsetPerTileX * StartTileX;
+					for (int32 TileX = StartTileX; TileX < EndTileX; TileX++)
 					{
-						UE_LOG(LogImgMedia, Error, TEXT("Could not load %s"), *Image);
+						// Get for our frame/mip level.
+						if (bHasTiles)
+						{
+							Image = FString::Printf(TEXT("%s_x%d_y%d.exr"), *BaseImage, TileX, TileY);
+						}
+						FRgbaInputFile InputFile(Image, 2);
+						if (InputFile.HasInputFile())
+						{
+							// read frame data
+							InputFile.SetFrameBuffer(MipDataPtr + FrameBufferOffsetX + FrameBufferOffsetY, Dim);
+							InputFile.ReadPixels(0, TileHeight - 1);
+
+							FrameBufferOffsetX += FrameBufferOffsetPerTileX;
+
+							if (!OutFrame->MipTilesPresent.Contains(CurrentMipLevel))
+							{
+								OutFrame->MipTilesPresent.Emplace(CurrentMipLevel, CurrentTileSelection);
+							}
+						}
+						else
+						{
+							UE_LOG(LogImgMedia, Error, TEXT("Could not load %s"), *Image);
+						}
 					}
+
+					FrameBufferOffsetY += FrameBufferOffsetPerTileY;
 				}
-				
-				FrameBufferOffsetY += FrameBufferOffsetPerTileY;
 			}
 		}
 
 		// Next level.
 		MipDataPtr += Dim.X * Dim.Y * BytesPerPixel;
 		Dim /= 2;
-		if (IsThisLevelPresent)
-		{
-			LevelFoundSoFar = true;
-		}
 	}
 
 	return true;
@@ -318,6 +326,12 @@ bool FExrImgMediaReader::GetInfo(FRgbaInputFile& InputFile, FImgMediaFrameInfo&
 	{
 		OutInfo.NumTiles = FIntPoint(OutInfo.Dim.X / OutInfo.TileDimensions.X, OutInfo.Dim.Y / OutInfo.TileDimensions.Y);
 	}
+	else
+	{
+		OutInfo.TileDimensions = OutInfo.Dim;
+		OutInfo.NumTiles = FIntPoint(1, 1);
+	}
+
 	return (OutInfo.UncompressedSize > 0) && (OutInfo.Dim.GetMin() > 0);
 }
 
diff --git a/Engine/Plugins/Media/ImgMedia/Source/ImgMedia/Private/Readers/ExrImgMediaReader.h b/Engine/Plugins/Media/ImgMedia/Source/ImgMedia/Private/Readers/ExrImgMediaReader.h
--- a/Engine/Plugins/Media/ImgMedia/Source/ImgMedia/Private/Readers/ExrImgMediaReader.h
+++ b/Engine/Plugins/Media/ImgMedia/Source/ImgMedia/Private/Readers/ExrImgMediaReader.h
@@ -28,7 +28,7 @@ public:
 	//~ IImgMediaReader interface
 
 	virtual bool GetFrameInfo(const FString& ImagePath, FImgMediaFrameInfo& OutInfo) override;
-	virtual bool ReadFrame(int32 FrameId, int32 MipLevel, const FImgMediaTileSelection& InTileSelection, TSharedPtr<FImgMediaFrame, ESPMode::ThreadSafe> OutFrame) override;
+	virtual bool ReadFrame(int32 FrameId, const TMap<int32, FImgMediaTileSelection>& InMipTiles, TSharedPtr<FImgMediaFrame, ESPMode::ThreadSafe> OutFrame) override;
 	virtual void CancelFrame(int32 FrameNumber) override;
 
 public:
diff --git a/Engine/Plugins/Media/ImgMedia/Source/ImgMedia/Private/Readers/ExrImgMediaReaderGpu.cpp b/Engine/Plugins/Media/ImgMedia/Source/ImgMedia/Private/Readers/ExrImgMediaReaderGpu.cpp
--- a/Engine/Plugins/Media/ImgMedia/Source/ImgMedia/Private/Readers/ExrImgMediaReaderGpu.cpp
+++ b/Engine/Plugins/Media/ImgMedia/Source/ImgMedia/Private/Readers/ExrImgMediaReaderGpu.cpp
@@ -121,12 +121,12 @@ FExrImgMediaReaderGpu::~FExrImgMediaReaderGpu()
 /* FExrImgMediaReaderGpu interface
  *****************************************************************************/
 
-bool FExrImgMediaReaderGpu::ReadFrame(int32 FrameId, int32 MipLevel, const FImgMediaTileSelection& InTileSelection, TSharedPtr<FImgMediaFrame, ESPMode::ThreadSafe> OutFrame)
+bool FExrImgMediaReaderGpu::ReadFrame(int32 FrameId, const TMap<int32, FImgMediaTileSelection>& InMipTiles, TSharedPtr<FImgMediaFrame, ESPMode::ThreadSafe> OutFrame)
 {
 	// Fall back to cpu?
 	if (bFallBackToCPU)
 	{
-		return FExrImgMediaReader::ReadFrame(FrameId, MipLevel, InTileSelection, OutFrame);
+		return FExrImgMediaReader::ReadFrame(FrameId, InMipTiles, OutFrame);
 	}
 
 	TSharedPtr<FImgMediaLoader, ESPMode::ThreadSafe> Loader = LoaderPtr.Pin();
@@ -146,10 +146,6 @@ bool FExrImgMediaReaderGpu::ReadFrame(int32 FrameId, int32 MipLevel, const FImgM
 
 	// Get tile info.
 	bool bHasTiles = OutFrame->Info.bHasTiles;
-	int32 StartTileX = InTileSelection.TopLeftX;
-	int32 StartTileY = InTileSelection.TopLeftY;
-	int32 EndTileX = FMath::Min((int32)InTileSelection.BottomRightX, OutFrame->Info.NumTiles.X);
-	int32 EndTileY = FMath::Min((int32)InTileSelection.BottomRightY, OutFrame->Info.NumTiles.Y);
 
 	const FIntPoint& FullResolution = OutFrame->Info.Dim;
 
@@ -160,12 +156,8 @@ bool FExrImgMediaReaderGpu::ReadFrame(int32 FrameId, int32 MipLevel, const FImgM
 		return false;
 	}
 
-	FIntRect TopMipViewport = FIntRect(FIntPoint(0, 0), FullResolution);
+	TMap<int32, FIntRect> Viewports;
 
-	if (bHasTiles)
-	{
-		TopMipViewport = FIntRect(StartTileX * TileDim.X, StartTileY * TileDim.Y, EndTileX * TileDim.X, EndTileY * TileDim.Y);
-	}
 	const int32 NumChannels = OutFrame->Info.NumChannels;
 	const int32 PixelSize = sizeof(uint16) * NumChannels;
 
@@ -177,18 +169,29 @@ bool FExrImgMediaReaderGpu::ReadFrame(int32 FrameId, int32 MipLevel, const FImgM
 		// Loop over all mips.
 		FIntPoint CurrentMipDim = FullResolution;
 
-		for (int32 CurrentMipLevel = 0; CurrentMipLevel < NumMipLevels; ++CurrentMipLevel)
+		for (const TPair<int32, FImgMediaTileSelection>& TilesPerMip : InMipTiles)
 		{
-			// Do we want to read in this mip?
-			bool IsThisLevelPresent = (OutFrame->MipMapsPresent & (1 << CurrentMipLevel)) != 0;
-			bool ReadThisMip = (CurrentMipLevel >= MipLevel) && (IsThisLevelPresent == false);
+			const int32 CurrentMipLevel = TilesPerMip.Key;
+			const FImgMediaTileSelection& CurrentTileSelection = TilesPerMip.Value;
 
+			bool ReadThisMip = true;
+			// Avoid reads if the cached frame already contains the current tiles for this mip level.
+			if (const FImgMediaTileSelection* CachedSelection = OutFrame->MipTilesPresent.Find(CurrentMipLevel))
+			{
+				ReadThisMip = !CachedSelection->Contains(CurrentTileSelection);
+			}
+			
 			// Next mip level.
 			int MipLevelDiv = 1 << CurrentMipLevel;
 			CurrentMipDim = FullResolution / MipLevelDiv;
 
 			if (ReadThisMip)
 			{
+				FIntRect& Viewport = Viewports.Add(CurrentMipLevel);
+				Viewport.Min = FIntPoint(TileDim.X * CurrentTileSelection.TopLeftX, TileDim.X * CurrentTileSelection.TopLeftY);
+				Viewport.Max = FIntPoint(TileDim.Y * CurrentTileSelection.BottomRightX, TileDim.Y * CurrentTileSelection.BottomRightY);
+				Viewport.Clip(FIntRect(FIntPoint::ZeroValue, CurrentMipDim));
+
 				const SIZE_T BufferSize = GetBufferSize(CurrentMipDim, NumChannels, bHasTiles, OutFrame->Info.NumTiles / MipLevelDiv);
 				FStructuredBufferPoolItemSharedPtr& BufferData = BufferDataArray[CurrentMipLevel];
 				BufferData = AllocateGpuBufferFromPool(BufferSize);
@@ -207,10 +210,10 @@ bool FExrImgMediaReaderGpu::ReadFrame(int32 FrameId, int32 MipLevel, const FImgM
 					if (bHasTiles)
 					{
 						FIntRect TileRegion = FIntRect(
-							FMath::FloorToInt(float(StartTileX) / MipLevelDiv),
-							FMath::FloorToInt(float(StartTileY) / MipLevelDiv),
-							FMath::CeilToInt(float(EndTileX) / MipLevelDiv),
-							FMath::CeilToInt(float(EndTileY) / MipLevelDiv));
+							(int32)CurrentTileSelection.TopLeftX,
+							(int32)CurrentTileSelection.TopLeftY,
+							FMath::Min((int32)CurrentTileSelection.BottomRightX, FMath::CeilToInt(float(OutFrame->Info.NumTiles.X) / MipLevelDiv)),
+							FMath::Min((int32)CurrentTileSelection.BottomRightY, FMath::CeilToInt(float(OutFrame->Info.NumTiles.Y) / MipLevelDiv)));
 						FIntPoint NumTiles = OutFrame->Info.NumTiles / MipLevelDiv;
 
 						ReadResult = ReadTiles(MipDataPtr, ImagePath, FrameId, TileRegion, NumTiles, TileDim, PixelSize);
@@ -222,7 +225,7 @@ bool FExrImgMediaReaderGpu::ReadFrame(int32 FrameId, int32 MipLevel, const FImgM
 
 					if (ReadResult != Fail)
 					{
-						OutFrame->MipMapsPresent |= 1 << CurrentMipLevel;
+						OutFrame->MipTilesPresent.Emplace(CurrentMipLevel, CurrentTileSelection);
 					}
 				}
 				else
@@ -264,7 +267,7 @@ bool FExrImgMediaReaderGpu::ReadFrame(int32 FrameId, int32 MipLevel, const FImgM
 
 					// Fall back to CPU.
 					bFallBackToCPU = true;
-					return FExrImgMediaReader::ReadFrame(FrameId, MipLevel, InTileSelection, OutFrame);
+					return FExrImgMediaReader::ReadFrame(FrameId, InMipTiles, OutFrame);
 				}
 			}
 		}
@@ -272,9 +275,8 @@ bool FExrImgMediaReaderGpu::ReadFrame(int32 FrameId, int32 MipLevel, const FImgM
 	OutFrame->Format = NumChannels <= 3 ? EMediaTextureSampleFormat::FloatRGB : EMediaTextureSampleFormat::FloatRGBA;
 	OutFrame->Stride = FullResolution.X * PixelSize;
 
-	OutFrame->SampleConverter = CreateSampleConverter(MoveTemp(BufferDataArray), FullResolution, TileDim, TopMipViewport, NumChannels, NumMipLevels, bHasTiles);
+	OutFrame->SampleConverter = CreateSampleConverter(MoveTemp(BufferDataArray), FullResolution, TileDim, Viewports, NumChannels, NumMipLevels, bHasTiles);
 
-	OutFrame->MipMapsPresent = 1 << MipLevel;
 	UE_LOG(LogImgMedia, Verbose, TEXT("Reader %p: Read Pixels Complete. %i"), this, FrameId);
 	return true;
 }
@@ -435,12 +437,12 @@ TSharedPtr<IMediaTextureSampleConverter, ESPMode::ThreadSafe> FExrImgMediaReader
 	( TArray<FStructuredBufferPoolItemSharedPtr>&& BufferDataArray
 	, const FIntPoint& FullResolution
 	, const FIntPoint& TileDim
-	, const FIntRect& TopMipViewport
+	, const TMap<int32, FIntRect>& Viewports
 	, const int32 NumChannels
 	, const int32 NumMipLevels
 	, const bool bHasTiles)
 {
-	auto RenderThreadSwizzler = [this, BufferDataArray, FullResolution, TileDim, TopMipViewport, NumChannels, NumMipLevels, bHasTiles](FRHICommandListImmediate& RHICmdList, FTexture2DRHIRef RenderTargetTextureRHI)->bool
+	auto RenderThreadSwizzler = [this, BufferDataArray, FullResolution, TileDim, Viewports, NumChannels, NumMipLevels, bHasTiles](FRHICommandListImmediate& RHICmdList, FTexture2DRHIRef RenderTargetTextureRHI)->bool
 	{
 
 		SCOPED_DRAW_EVENT(RHICmdList, FExrImgMediaReaderGpu_Convert);
@@ -452,8 +454,6 @@ TSharedPtr<IMediaTextureSampleConverter, ESPMode::ThreadSafe> FExrImgMediaReader
 			int MipLevelDiv = 1 << MipLevel;
 			Dim = FullResolution / MipLevelDiv;
 
-			FIntRect Viewport = TopMipViewport / MipLevelDiv;
-			Viewport.Clip(FIntRect(FIntPoint(0, 0), Dim));
 			FStructuredBufferPoolItemSharedPtr BufferData = BufferDataArray[MipLevel];
 			if (BufferData.IsValid())
 			{
@@ -461,6 +461,9 @@ TSharedPtr<IMediaTextureSampleConverter, ESPMode::ThreadSafe> FExrImgMediaReader
 				{
 					continue;
 				}
+
+				FIntRect Viewport = Viewports[MipLevel];
+
 				// This flag will indicate that we should wait for poll to complete.
 				BufferData->bWillBeSignaled = true;
 
diff --git a/Engine/Plugins/Media/ImgMedia/Source/ImgMedia/Private/Readers/ExrImgMediaReaderGpu.h b/Engine/Plugins/Media/ImgMedia/Source/ImgMedia/Private/Readers/ExrImgMediaReaderGpu.h
--- a/Engine/Plugins/Media/ImgMedia/Source/ImgMedia/Private/Readers/ExrImgMediaReaderGpu.h
+++ b/Engine/Plugins/Media/ImgMedia/Source/ImgMedia/Private/Readers/ExrImgMediaReaderGpu.h
@@ -10,7 +10,6 @@
 #include "IMediaTextureSampleConverter.h"
 
 class FExrImgMediaReaderGpu;
-struct FImgMediaTileSelection;
 
 struct FStructuredBufferPoolItem
 {
@@ -63,7 +62,7 @@ public:
 public:
 
 	//~ FExrImgMediaReader interface
-	virtual bool ReadFrame(int32 FrameId, int32 MipLevel, const FImgMediaTileSelection& InTileSelection, TSharedPtr<FImgMediaFrame, ESPMode::ThreadSafe> OutFrame) override;
+	virtual bool ReadFrame(int32 FrameId, const TMap<int32, FImgMediaTileSelection>& InMipTiles, TSharedPtr<FImgMediaFrame, ESPMode::ThreadSafe> OutFrame) override;
 	
 	/**
 	* For performance reasons we want to pre-allocate structured buffers to at least the number of concurrent frames.
@@ -105,7 +104,7 @@ protected:
 		( TArray<FStructuredBufferPoolItemSharedPtr>&& BufferDataArray
 		, const FIntPoint& FullResolution
 		, const FIntPoint& TileDim
-		, const FIntRect& TopMipViewport
+		, const TMap<int32, FIntRect>& Viewports
 		, const int32 NumChannels
 		, const int32 NumMipLevels
 		, const bool bHasTiles);
diff --git a/Engine/Plugins/Media/ImgMedia/Source/ImgMedia/Private/Readers/GenericImgMediaReader.cpp b/Engine/Plugins/Media/ImgMedia/Source/ImgMedia/Private/Readers/GenericImgMediaReader.cpp
--- a/Engine/Plugins/Media/ImgMedia/Source/ImgMedia/Private/Readers/GenericImgMediaReader.cpp
+++ b/Engine/Plugins/Media/ImgMedia/Source/ImgMedia/Private/Readers/GenericImgMediaReader.cpp
@@ -91,7 +91,7 @@ bool FGenericImgMediaReader::GetFrameInfo(const FString& ImagePath, FImgMediaFra
 }
 
 
-bool FGenericImgMediaReader::ReadFrame(int32 FrameId, int32 MipLevel, const FImgMediaTileSelection& InTileSelectio, TSharedPtr<FImgMediaFrame, ESPMode::ThreadSafe> OutFrame)
+bool FGenericImgMediaReader::ReadFrame(int32 FrameId, const TMap<int32, FImgMediaTileSelection>& InMipTiles, TSharedPtr<FImgMediaFrame, ESPMode::ThreadSafe> OutFrame)
 {
 	TSharedPtr<FImgMediaLoader, ESPMode::ThreadSafe> Loader = LoaderPtr.Pin();
 	if (Loader.IsValid() == false)
@@ -106,12 +106,13 @@ bool FGenericImgMediaReader::ReadFrame(int32 FrameId, int32 MipLevel, const FImg
 	void* Buffer = OutFrame->Data.Get();
 
 	// Loop over all mips.
-	for (int32 CurrentMipLevel = 0; CurrentMipLevel < NumMipLevels; ++CurrentMipLevel)
+	for (const TPair<int32, FImgMediaTileSelection>& TilesPerMip : InMipTiles)
 	{
+		const int32 CurrentMipLevel = TilesPerMip.Key;
+
 		// Do we want to read in this mip?
-		bool IsThisLevelPresent = (OutFrame->MipMapsPresent & (1 << CurrentMipLevel)) != 0;
-		bool ReadThisMip = (Buffer == nullptr) ||
-			((CurrentMipLevel >= MipLevel) && (IsThisLevelPresent == false));
+		bool IsThisLevelPresent = OutFrame->MipTilesPresent.Contains(CurrentMipLevel);
+		bool ReadThisMip = (Buffer == nullptr) || (IsThisLevelPresent == false);
 		if (ReadThisMip)
 		{
 			// Load image.
@@ -153,7 +154,7 @@ bool FGenericImgMediaReader::ReadFrame(int32 FrameId, int32 MipLevel, const FImg
 
 			// Copy data to our buffer
 			FMemory::Memcpy(Buffer, RawData.GetData(), RawNum);
-			OutFrame->MipMapsPresent |= 1 << CurrentMipLevel;
+			OutFrame->MipTilesPresent.Add(CurrentMipLevel).SetAllVisible();
 		}
 
 		// Next level.
diff --git a/Engine/Plugins/Media/ImgMedia/Source/ImgMedia/Private/Readers/GenericImgMediaReader.h b/Engine/Plugins/Media/ImgMedia/Source/ImgMedia/Private/Readers/GenericImgMediaReader.h
--- a/Engine/Plugins/Media/ImgMedia/Source/ImgMedia/Private/Readers/GenericImgMediaReader.h
+++ b/Engine/Plugins/Media/ImgMedia/Source/ImgMedia/Private/Readers/GenericImgMediaReader.h
@@ -32,7 +32,7 @@ public:
 	//~ IImgMediaReader interface
 
 	virtual bool GetFrameInfo(const FString& ImagePath, FImgMediaFrameInfo& OutInfo) override;
-	virtual bool ReadFrame(int32 FrameId, int32 MipLevel, const FImgMediaTileSelection& InTileSelectio, TSharedPtr<FImgMediaFrame, ESPMode::ThreadSafe> OutFrame) override;
+	virtual bool ReadFrame(int32 FrameId, const TMap<int32, FImgMediaTileSelection>& InMipTiles, TSharedPtr<FImgMediaFrame, ESPMode::ThreadSafe> OutFrame) override;
 	virtual void CancelFrame(int32 FrameNumber) override {};
 
 private:
diff --git a/Engine/Plugins/Media/ImgMedia/Source/ImgMedia/Private/Readers/IImgMediaReader.h b/Engine/Plugins/Media/ImgMedia/Source/ImgMedia/Private/Readers/IImgMediaReader.h
--- a/Engine/Plugins/Media/ImgMedia/Source/ImgMedia/Private/Readers/IImgMediaReader.h
+++ b/Engine/Plugins/Media/ImgMedia/Source/ImgMedia/Private/Readers/IImgMediaReader.h
@@ -5,13 +5,12 @@
 #include "CoreTypes.h"
 #include "Misc/FrameRate.h"
 #include "IMediaTextureSample.h"
+#include "ImgMediaMipMapInfo.h"
 #include "Math/IntPoint.h"
 #include "Templates/SharedPointer.h"
 #include "RHI.h"
 #include "RHIResources.h"
 
-struct FImgMediaTileSelection;
-
 /**
  * Information about an image sequence frame.
  */
@@ -66,8 +65,8 @@ struct FImgMediaFrame
 	/** Additional information about the frame. */
 	FImgMediaFrameInfo Info;
 
-	/** Bitmask of which mipmaps are present. */
-	int32 MipMapsPresent = 0;
+	/** Tiles present per mip level. */
+	TMap<int32, FImgMediaTileSelection> MipTilesPresent;
 
 	/** The frame's horizontal stride (in bytes). */
 	uint32 Stride = 0;
@@ -119,7 +118,7 @@ public:
 	 * @return true on success, false otherwise.
 	 * @see GetFrameInfo
 	 */
-	virtual bool ReadFrame(int32 FrameId, int32 MipLevel, const FImgMediaTileSelection& InTileSelection, TSharedPtr<FImgMediaFrame, ESPMode::ThreadSafe> OutFrame) = 0;
+	virtual bool ReadFrame(int32 FrameId, const TMap<int32, FImgMediaTileSelection>& InMipTiles, TSharedPtr<FImgMediaFrame, ESPMode::ThreadSafe> OutFrame) = 0;
 
 	/**
 	 * Mark Frame to be canceled based on Frame number. Typically this will be 
diff --git a/Engine/Plugins/Media/ImgMedia/Source/ImgMedia/Public/ImgMediaMipMapInfo.h b/Engine/Plugins/Media/ImgMedia/Source/ImgMedia/Public/ImgMediaMipMapInfo.h
--- a/Engine/Plugins/Media/ImgMedia/Source/ImgMedia/Public/ImgMediaMipMapInfo.h
+++ b/Engine/Plugins/Media/ImgMedia/Source/ImgMedia/Public/ImgMediaMipMapInfo.h
@@ -92,18 +92,62 @@ struct FImgMediaTileSelection
 	 */
 	void SetAllNotVisible()
 	{
-		TopLeftX = 0;
-		TopLeftY = 0;
+		TopLeftX = 0xffff;
+		TopLeftY = 0xffff;
 		BottomRightX = 0;
 		BottomRightY = 0;
 	}
 
 	/**
-	 * See if this selection is visible.
+	* See if this selection is visible.
+	*
+	* @return True if so.
+	*/
+	bool IsVisible() const { return (TopLeftX < BottomRightX); };
+
+
+	/**
+	 * Include a given tile coordinate to the current section region.
 	 *
-	 * @return True if so.
+	 * @param TileCoordX Horizontal tile coordinate.
+	 * @param TileCoordY Vertical tile coordinate.
 	 */
-	bool IsVisible() { return (TopLeftX < BottomRightX); };
+	void Include(uint16 TileCoordX, uint16 TileCoordY)
+	{
+		TopLeftX = FMath::Min(TopLeftX, TileCoordX);
+		TopLeftY = FMath::Min(TopLeftY, TileCoordY);
+		BottomRightX = FMath::Max(BottomRightX, uint16(TileCoordX + 1u));
+		BottomRightY = FMath::Max(BottomRightY, uint16(TileCoordY + 1u));
+	}
+
+	/**
+	 * Check if the current selection contains a tile.
+	 *
+	 * @param TileCoordX Horizontal tile coordinate.
+	 * @param TileCoordY Vertical tile coordinate.
+	 * @return True if the coordinate is contained withing bounds, false otherwise.
+	 */
+	bool Contains(uint16 TileCoordX, uint16 TileCoordY) const
+	{
+		return TopLeftX <= TileCoordX &&
+			TopLeftY <= TileCoordY &&
+			BottomRightX > TileCoordX &&
+			BottomRightY > TileCoordY;
+	}
+
+	/**
+	 * Check if the current selection contains another selection within its bounds.
+	 *
+	 * @param Other Selection to compare.
+	 * @return True if the other selection is contained withing bounds, false otherwise.
+	 */
+	bool Contains(const FImgMediaTileSelection& Other) const
+	{
+		return TopLeftX <= Other.TopLeftX &&
+			TopLeftY <= Other.TopLeftY &&
+			BottomRightX >= Other.BottomRightX &&
+			BottomRightY >= Other.BottomRightY;
+	}
 };
 
 /**
diff --git a/Engine/Plugins/Media/ImgMedia/Source/ImgMediaFactory/Private/ImgMediaSettings.cpp b/Engine/Plugins/Media/ImgMedia/Source/ImgMediaFactory/Private/ImgMediaSettings.cpp
--- a/Engine/Plugins/Media/ImgMedia/Source/ImgMediaFactory/Private/ImgMediaSettings.cpp
+++ b/Engine/Plugins/Media/ImgMedia/Source/ImgMediaFactory/Private/ImgMediaSettings.cpp
@@ -8,6 +8,7 @@
 
 UImgMediaSettings::UImgMediaSettings()
 	: DefaultFrameRate(24, 1)
+	, bUseVirtualTextureTiles(false)
 	, CacheBehindPercentage(25)
 	, CacheSizeGB(1.0f)
 	, CacheThreads(8)
diff --git a/Engine/Plugins/Media/ImgMedia/Source/ImgMediaFactory/Public/ImgMediaSettings.h b/Engine/Plugins/Media/ImgMedia/Source/ImgMediaFactory/Public/ImgMediaSettings.h
--- a/Engine/Plugins/Media/ImgMedia/Source/ImgMediaFactory/Public/ImgMediaSettings.h
+++ b/Engine/Plugins/Media/ImgMedia/Source/ImgMediaFactory/Public/ImgMediaSettings.h
@@ -26,6 +26,9 @@ public:
 	UPROPERTY(config, EditAnywhere, Category=General)
 	FFrameRate DefaultFrameRate;
 
+	/** Whether to use the tiles produced by virtual texturing or not. */
+	UPROPERTY(config, EditAnywhere, Category=General)
+	bool bUseVirtualTextureTiles;
 public:
 
 	/** Percentage of cache to use for frames behind the play head (default = 25%). */
