ImgMedia: Tells MediaSource that it can play exrs, etc.

#rnx

diff --git a/Engine/Plugins/Media/ImgMedia/Source/ImgMedia/Private/ImgMediaModule.cpp b/Engine/Plugins/Media/ImgMedia/Source/ImgMedia/Private/ImgMediaModule.cpp
--- a/Engine/Plugins/Media/ImgMedia/Source/ImgMedia/Private/ImgMediaModule.cpp
+++ b/Engine/Plugins/Media/ImgMedia/Source/ImgMedia/Private/ImgMediaModule.cpp
@@ -10,6 +10,7 @@
 #include "ImgMediaGlobalCache.h"
 #include "ImgMediaPlayer.h"
 #include "ImgMediaScheduler.h"
+#include "ImgMediaSource.h"
 #include "IImgMediaModule.h"
 
 
@@ -124,11 +125,23 @@ public:
 
 	virtual void StartupModule() override
 	{
-
+		// Register media source spawners.
+		FMediaSourceSpawnDelegate SpawnDelegate =
+			FMediaSourceSpawnDelegate::CreateStatic(&FImgMediaModule::SpawnMediaSourceForString);
+		for (const FString& Ext : FileExtensions)
+		{
+			UMediaSource::RegisterSpawnFromFileExtension(Ext, SpawnDelegate);
+		}
 	}
 
 	virtual void ShutdownModule() override
 	{
+		// Unregister media source spawners.
+		for (const FString& Ext : FileExtensions)
+		{
+			UMediaSource::UnregisterSpawnFromFileExtension(Ext);
+		}
+
 		Scheduler.Reset();
 		GlobalCache.Reset();
 
@@ -160,7 +173,32 @@ private:
 		GlobalCache->Initialize();
 	}
 
+	/**
+	 * Creates a media source for MediaPath.
+	 *
+	 * @param	MediaPath		File path to the media.
+	 */
+	static UMediaSource* SpawnMediaSourceForString(const FString& MediaPath)
+	{
+		TObjectPtr<UImgMediaSource> MediaSource = 
+			NewObject<UImgMediaSource>(GetTransientPackage(), NAME_None,
+				RF_Transactional | RF_Transient);
+		MediaSource->SetSequencePath(MediaPath);
+
+		return MediaSource;
+	}
+
 	TSharedPtr<FImgMediaScheduler, ESPMode::ThreadSafe> Scheduler;
+
+	/** List of file extensions that we support. */
+	const TArray<FString> FileExtensions =
+	{
+		TEXT("bmp"),
+		TEXT("exr"),
+		TEXT("jpg"),
+		TEXT("jpeg"),
+		TEXT("png"),
+	};
 };
 
 
