ImgMedia: Process Images can now output mipmapped files.

#rnx

diff --git a/Engine/Plugins/Media/ImgMedia/Source/ImgMediaEditor/Private/Widgets/ImgMediaProcessImagesOptions.h b/Engine/Plugins/Media/ImgMedia/Source/ImgMediaEditor/Private/Widgets/ImgMediaProcessImagesOptions.h
--- a/Engine/Plugins/Media/ImgMedia/Source/ImgMediaEditor/Private/Widgets/ImgMediaProcessImagesOptions.h
+++ b/Engine/Plugins/Media/ImgMedia/Source/ImgMediaEditor/Private/Widgets/ImgMediaProcessImagesOptions.h
@@ -26,6 +26,10 @@ public:
 	/** If checked, then use our custom format with a single file for all tiles. */
 	UPROPERTY(EditAnywhere, Category = Sequence)
 	bool bUseCustomFormat;
+
+	/** If checked, then enable mip mapping. */
+	UPROPERTY(EditAnywhere, Category = Sequence)
+	bool bEnableMipMapping;
 	
 	/** Width of a tile in pixels. If 0, then do not make tiles. */
 	UPROPERTY(EditAnywhere, Transient, Category = Tiles)
diff --git a/Engine/Plugins/Media/ImgMedia/Source/ImgMediaEditor/Private/Widgets/SImgMediaProcessImages.cpp b/Engine/Plugins/Media/ImgMedia/Source/ImgMediaEditor/Private/Widgets/SImgMediaProcessImages.cpp
--- a/Engine/Plugins/Media/ImgMedia/Source/ImgMediaEditor/Private/Widgets/SImgMediaProcessImages.cpp
+++ b/Engine/Plugins/Media/ImgMedia/Source/ImgMediaEditor/Private/Widgets/SImgMediaProcessImages.cpp
@@ -108,6 +108,7 @@ void SImgMediaProcessImages::ProcessAllImages(TSharedPtr<SNotificationItem> Conf
 	int32 InTileWidth = Options->TileSizeX;
 	int32 InTileHeight = Options->TileSizeY;
 	int32 TileBorder = Options->TileBorder;
+	bool bEnableMips = Options->bEnableMipMapping;
 
 	// Create output directory.
 	FString OutPath = Options->OutputPath.Path;
@@ -176,7 +177,8 @@ void SImgMediaProcessImages::ProcessAllImages(TSharedPtr<SNotificationItem> Conf
 				FString Name = FPaths::Combine(OutPath, FileName);
 				if (bUseCustomFormat)
 				{
-					ProcessImageCustom(ImageWrapper, InTileWidth, InTileHeight, TileBorder, Name);
+					ProcessImageCustom(ImageWrapper, InTileWidth, InTileHeight, TileBorder,
+						bEnableMips, Name);
 				}
 				else
 				{
@@ -264,7 +266,8 @@ void SImgMediaProcessImages::ProcessImage(TSharedPtr<IImageWrapper>& InImageWrap
 }
 
 void SImgMediaProcessImages::ProcessImageCustom(TSharedPtr<IImageWrapper>& InImageWrapper,
-	int32 InTileWidth, int32 InTileHeight, int32 InTileBorder, const FString& InName)
+	int32 InTileWidth, int32 InTileHeight, int32 InTileBorder, bool bInEnableMips,
+	const FString& InName)
 {
 #if IMGMEDIAEDITOR_EXR_SUPPORTED_PLATFORM
 	// Get image data.
@@ -374,7 +377,7 @@ void SImgMediaProcessImages::ProcessImageCustom(TSharedPtr<IImageWrapper>& InIma
 
 	int32 NumChannels = 4;
 
-	FIntPoint Stride(2, DestWidth * BytesPerPixel);
+	FIntPoint Stride(2, 0);
 
 	// Create tiled exr file.
 	FTiledOutputFile OutFile(FIntPoint(0, 0), FIntPoint(DestWidth - 1, DestHeight - 1),
@@ -398,14 +401,83 @@ void SImgMediaProcessImages::ProcessImageCustom(TSharedPtr<IImageWrapper>& InIma
 	OutFile.AddChannel(RChannelName);
 
 	// Create output.
-	OutFile.CreateOutputFile(InName, DestWidth, DestHeight, false);
-	OutFile.AddFrameBufferChannel(AChannelName, RawDataPtr, Stride);
-	OutFile.AddFrameBufferChannel(BChannelName, RawDataPtr + DestWidth * 2, Stride);
-	OutFile.AddFrameBufferChannel(GChannelName, RawDataPtr + DestWidth * 4, Stride);
-	OutFile.AddFrameBufferChannel(RChannelName, RawDataPtr + DestWidth * 6, Stride);
-	OutFile.SetFrameBuffer();
-	
-	OutFile.WriteTile(0, 0, 0);
+	OutFile.CreateOutputFile(InName, DestWidth, DestHeight, bInEnableMips);
+	OutFile.AddFrameBufferChannel(AChannelName, nullptr, Stride);
+	OutFile.AddFrameBufferChannel(BChannelName, nullptr, Stride);
+	OutFile.AddFrameBufferChannel(GChannelName, nullptr, Stride);
+	OutFile.AddFrameBufferChannel(RChannelName, nullptr, Stride);
+
+	// Flip between 2 buffers making mips.
+	TArray64<uint8> RawData2;
+	uint8* MipBuffer[2];
+	MipBuffer[0] = RawDataPtr;
+	MipBuffer[1] = nullptr;
+	int32 CurrentMipBufferIndex = 0;
+
+	// Loop over each mip level.
+	int32 NumMips = OutFile.GetNumberOfMipLevels();
+	for (int32 MipLevel = 0; MipLevel < NumMips; MipLevel++)
+	{
+		int32 MipWidth = OutFile.GetMipWidth(MipLevel);
+		int32 MipHeight = OutFile.GetMipHeight(MipLevel);
+		uint8* CurrentBuffer = MipBuffer[CurrentMipBufferIndex];
+		uint8* LastBuffer = MipBuffer[CurrentMipBufferIndex ^ 1];
+
+		// Allocate space for the other buffer.
+		if (CurrentBuffer == nullptr)
+		{
+			RawData2.AddUninitialized(MipWidth* MipHeight* BytesPerPixel);
+			MipBuffer[CurrentMipBufferIndex] = RawData2.GetData();
+			CurrentBuffer = MipBuffer[CurrentMipBufferIndex];
+		}
+
+		// Generate mip data.
+		if (MipLevel != 0)
+		{
+			int32 SourceStrideX = NumChannels;
+			int32 SourceStrideY = MipWidth * NumChannels * 2;
+			for (int32 PixelY = 0; PixelY < MipHeight; ++PixelY)
+			{
+				for (int32 PixelX = 0; PixelX < MipWidth; ++PixelX)
+				{
+					int32 PixelOffset = (PixelX + PixelY * MipWidth) * NumChannels;
+					for (int32 Channel = 0; Channel < NumChannels; ++Channel)
+					{
+						// Box filter.
+						int32 SourceWidth = MipWidth * 2;
+						int32 SourcePixelOffset = (PixelX + PixelY * SourceWidth) * NumChannels * 2 +
+							Channel;
+						FFloat16 SourcePixels[4];
+						SourcePixels[0] =
+							((FFloat16*)LastBuffer)[SourcePixelOffset];
+						SourcePixels[1] =
+							((FFloat16*)LastBuffer)[SourcePixelOffset + SourceStrideX];
+						SourcePixels[2] =
+							((FFloat16*)LastBuffer)[SourcePixelOffset + SourceStrideY];
+						SourcePixels[3] =
+							((FFloat16*)LastBuffer)[SourcePixelOffset + SourceStrideX + SourceStrideY];
+
+						((FFloat16*)CurrentBuffer)[PixelOffset + Channel] =
+							(SourcePixels[0] + SourcePixels[1] + SourcePixels[2] + SourcePixels[3]) * 0.25f;
+					}
+				}
+			}
+		}
+
+		// Write to EXR.
+		Stride.Y = MipWidth * BytesPerPixel;
+		OutFile.UpdateFrameBufferChannel(AChannelName, CurrentBuffer, Stride);
+		OutFile.UpdateFrameBufferChannel(BChannelName, CurrentBuffer + MipWidth * 2, Stride);
+		OutFile.UpdateFrameBufferChannel(GChannelName, CurrentBuffer + MipWidth * 4, Stride);
+		OutFile.UpdateFrameBufferChannel(RChannelName, CurrentBuffer + MipWidth * 6, Stride);
+
+		OutFile.SetFrameBuffer();
+
+		OutFile.WriteTile(0, 0, MipLevel);
+
+		// Switch buffers.
+		CurrentMipBufferIndex ^= 1;
+	}
 
 #else // IMGMEDIAEDITOR_EXR_SUPPORTED_PLATFORM
 	UE_LOG(LogImgMediaEditor, Error, TEXT("EXR not supported on this platform."));
diff --git a/Engine/Plugins/Media/ImgMedia/Source/ImgMediaEditor/Private/Widgets/SImgMediaProcessImages.h b/Engine/Plugins/Media/ImgMedia/Source/ImgMediaEditor/Private/Widgets/SImgMediaProcessImages.h
--- a/Engine/Plugins/Media/ImgMedia/Source/ImgMediaEditor/Private/Widgets/SImgMediaProcessImages.h
+++ b/Engine/Plugins/Media/ImgMedia/Source/ImgMediaEditor/Private/Widgets/SImgMediaProcessImages.h
@@ -59,10 +59,12 @@ private:
 	 * @param InTileWidth		Desired width of tiles.
 	 * @param InTileHeight		Desired height of tiles.
 	 * @param InTileBorder		Number of pixels to duplicate along a tile edge.
+	 * @param bInEnableMips		Turn on mip mapping.
 	 * @param InName			Full path and name of file to write.
 	 */
 	void ProcessImageCustom(TSharedPtr<IImageWrapper>& InImageWrapper,
-		int32 InTileWidth, int32 InTileHeight, int32 InTileBorder, const FString& InName);
+		int32 InTileWidth, int32 InTileHeight, int32 InTileBorder, 
+		bool bInEnableMips, const FString& InName);
 
 	/** Holds our details view. */
 	TSharedPtr<class IDetailsView> DetailsView;
diff --git a/Engine/Plugins/Media/ImgMedia/Source/OpenExrWrapper/Private/OpenExrWrapper.cpp b/Engine/Plugins/Media/ImgMedia/Source/OpenExrWrapper/Private/OpenExrWrapper.cpp
--- a/Engine/Plugins/Media/ImgMedia/Source/OpenExrWrapper/Private/OpenExrWrapper.cpp
+++ b/Engine/Plugins/Media/ImgMedia/Source/OpenExrWrapper/Private/OpenExrWrapper.cpp
@@ -460,6 +460,23 @@ void FTiledOutputFile::AddFrameBufferChannel(const FString& Name, void* Base,
 		Imf::Slice(Imf::HALF, (char*)Base, Stride.X, Stride.Y));
 }
 
+void FTiledOutputFile::UpdateFrameBufferChannel(const FString& Name, void* Base,
+	const FIntPoint& Stride)
+{
+	Imf::Slice* Slice = ((Imf::FrameBuffer*)FrameBuffer)->findSlice(
+		StringCast<ANSICHAR>(*Name).Get());
+	if (Slice != nullptr)
+	{
+		Slice->base = (char*)Base;
+		Slice->xStride = Stride.X;
+		Slice->yStride = Stride.Y;
+	}
+	else
+	{
+		UE_LOG(LogOpenEXRWrapper, Error, TEXT("Could not find frame buffer channel %s."), *Name);
+	}
+}
+
 void FTiledOutputFile::SetFrameBuffer()
 {
 	if (OutputFile != nullptr)
@@ -494,4 +511,46 @@ void FTiledOutputFile::WriteTile(int32 TileX, int32 TileY, int32 MipLevel)
 	}
 }
 
+int32 FTiledOutputFile::GetNumberOfMipLevels()
+{
+	if (OutputFile != nullptr)
+	{
+		return ((Imf::TiledOutputFile*)OutputFile)->numLevels();
+	}
+	else
+	{
+		UE_LOG(LogOpenEXRWrapper, Error,
+			TEXT("GetNumberOfMipLevels failed: CreateOutputFile has not been called yet."));
+		return 0;
+	}
+}
+
+int32 FTiledOutputFile::GetMipWidth(int32 MipLevel)
+{
+	if (OutputFile != nullptr)
+	{
+		return ((Imf::TiledOutputFile*)OutputFile)->levelWidth(MipLevel);
+	}
+	else
+	{
+		UE_LOG(LogOpenEXRWrapper, Error,
+			TEXT("GetMipWidth failed: CreateOutputFile has not been called yet."));
+		return 0;
+	}
+}
+
+int32 FTiledOutputFile::GetMipHeight(int32 MipLevel)
+{
+	if (OutputFile != nullptr)
+	{
+		return ((Imf::TiledOutputFile*)OutputFile)->levelHeight(MipLevel);
+	}
+	else
+	{
+		UE_LOG(LogOpenEXRWrapper, Error,
+			TEXT("GetMipHeight failed: CreateOutputFile has not been called yet."));
+		return 0;
+	}
+}
+
 IMPLEMENT_MODULE(FDefaultModuleImpl, OpenExrWrapper);
diff --git a/Engine/Plugins/Media/ImgMedia/Source/OpenExrWrapper/Public/OpenExrWrapper.h b/Engine/Plugins/Media/ImgMedia/Source/OpenExrWrapper/Public/OpenExrWrapper.h
--- a/Engine/Plugins/Media/ImgMedia/Source/OpenExrWrapper/Public/OpenExrWrapper.h
+++ b/Engine/Plugins/Media/ImgMedia/Source/OpenExrWrapper/Public/OpenExrWrapper.h
@@ -90,6 +90,11 @@ public:
 	 */
 	void AddIntAttribute(const FString& Name, int32 Value);
 
+	/**
+	 * Call this to get the number of mipmap levels.
+	 */
+	virtual int32 GetNumberOfMipLevels() { return 0; }
+
 protected:
 	/** Stores the EXR header object. */
 	void* Header;
@@ -133,11 +138,6 @@ public:
 	void CreateOutputFile(const FString& FilePath,
 		int32 TileWidth, int32 TileHeight, int32 NumChannels, bool bIsMipsEnabled);
 
-	/**
-	 * Call this to get the number of mipmap levels.
-	 */
-	int32 GetNumberOfMipLevels();
-
 	/**
 	 * Call this prior to WriteTile to set where the data is coming from.
 	 * 
@@ -154,6 +154,9 @@ public:
 	 * @param MipLevel		Mipmap level of tile.
 	 */
 	void WriteTile(int32 TileX, int32 TileY, int32 MipLevel);
+
+	//~ FBaseOutputFile interface.
+	virtual int32 GetNumberOfMipLevels() override;
 };
 
 /**
@@ -204,6 +207,16 @@ public:
 	void CreateOutputFile(const FString& FilePath,
 		int32 TileWidth, int32 TileHeight, bool bIsMipsEnabled);
 
+	/**
+	 * Get the width of a mip level.
+	 */
+	int32 GetMipWidth(int32 MipLevel);
+
+	/**
+	 * Get the height of a mip level.
+	 */
+	int32 GetMipHeight(int32 MipLevel);
+
 	/**
 	 * Call this before SetFrameBuffer for each channel in the frame buffer.
 	 *
@@ -213,6 +226,17 @@ public:
 	 */
 	void AddFrameBufferChannel(const FString& Name, void* Base,
 		const FIntPoint& Stride);
+
+	/**
+	 * Call this to change a channel in the frame buffer.
+	 *
+	 * @param Name			Name of channel (e.g. R, G, B, or A).
+	 * @param Base			Address of the start of the data.
+	 * @param Stride		A pixels location is calculated by Base + x * StrideX + y * StrideY.
+	 */
+	void UpdateFrameBufferChannel(const FString& Name, void* Base,
+		const FIntPoint& Stride);
+
 	/**
 	 * Call this prior to WriteTile to set where the data is coming from.
 	 *
@@ -230,6 +254,9 @@ public:
 	 */
 	void WriteTile(int32 TileX, int32 TileY, int32 MipLevel);
 
+	//~ FBaseOutputFile interface.
+	virtual int32 GetNumberOfMipLevels() override;
+
 private:
 
 	/** Stores the EXR frame buffer object. */
