ImgMedia: Fixed tile visibility calculation taking object size into account twice.

#rnx

diff --git a/Engine/Plugins/Media/ImgMedia/Source/ImgMedia/Private/Assets/ImgMediaMipMapInfo.cpp b/Engine/Plugins/Media/ImgMedia/Source/ImgMedia/Private/Assets/ImgMediaMipMapInfo.cpp
--- a/Engine/Plugins/Media/ImgMedia/Source/ImgMedia/Private/Assets/ImgMediaMipMapInfo.cpp
+++ b/Engine/Plugins/Media/ImgMedia/Source/ImgMedia/Private/Assets/ImgMediaMipMapInfo.cpp
@@ -395,6 +395,17 @@ void FImgMediaMipMapInfo::GetObjectSize(const AActor* InActor, float& Width, flo
 		InActor->GetActorBounds(false, BoundsOrigin, BoxExtent);
 		Width = BoxExtent.Z * 2.0f;
 		Height = BoxExtent.Y * 2.0f;
+
+		// We apply the object scsale in UpdateMipLevelCache, so take it out here.
+		FVector ObjectScale = InActor->GetActorScale();
+		if (FMath::IsNearlyZero(ObjectScale.X) == false)
+		{
+			Width = Width / ObjectScale.X;
+		}
+		if (FMath::IsNearlyZero(ObjectScale.Y) == false)
+		{
+			Height = Height / ObjectScale.Y;
+		}
 	}
 
 	// Did we get anything?
