ImgMedia: Added some limits to some process images options.
#rnx

diff --git a/Engine/Plugins/Media/ImgMedia/Source/ImgMediaEditor/Private/Widgets/ImgMediaProcessImagesOptions.h b/Engine/Plugins/Media/ImgMedia/Source/ImgMediaEditor/Private/Widgets/ImgMediaProcessImagesOptions.h
--- a/Engine/Plugins/Media/ImgMedia/Source/ImgMediaEditor/Private/Widgets/ImgMediaProcessImagesOptions.h
+++ b/Engine/Plugins/Media/ImgMedia/Source/ImgMediaEditor/Private/Widgets/ImgMediaProcessImagesOptions.h
@@ -32,11 +32,11 @@ public:
 	bool bEnableMipMapping;
 	
 	/** Width of a tile in pixels. If 0, then do not make tiles. */
-	UPROPERTY(EditAnywhere, Transient, Category = Tiles)
+	UPROPERTY(EditAnywhere, Transient, Category = Tiles, meta = (ClampMin = "0.0"))
 	int32 TileSizeX = 0;
 
 	/** Height of a tile in pixels. If 0, then do not make tiles. */
-	UPROPERTY(EditAnywhere, Transient, Category = Tiles)
+	UPROPERTY(EditAnywhere, Transient, Category = Tiles, meta = (ClampMin = "0.0"))
 	int32 TileSizeY = 0;
 
 	/** Number of pixels to duplicate along each tile border. */
