ImgMedia: Removed virtual texture tile module setting in favor of an (advanced) image media source option.

#rnx

diff --git a/Engine/Plugins/Media/ImgMedia/Source/ImgMedia/Private/Assets/ImgMediaSource.cpp b/Engine/Plugins/Media/ImgMedia/Source/ImgMedia/Private/Assets/ImgMediaSource.cpp
--- a/Engine/Plugins/Media/ImgMedia/Source/ImgMedia/Private/Assets/ImgMediaSource.cpp
+++ b/Engine/Plugins/Media/ImgMedia/Source/ImgMedia/Private/Assets/ImgMediaSource.cpp
@@ -19,6 +19,7 @@ UImgMediaSource::UImgMediaSource()
 	: IsPathRelativeToProjectRoot(false)
 	, FrameRateOverride(0, 0)
 	, bFillGapsInSequence(true)
+	, bReadVirtualTextureTiles(false)
 	, MipMapInfo(MakeShared<FImgMediaMipMapInfo, ESPMode::ThreadSafe>())
 {
 }
@@ -97,6 +98,10 @@ bool UImgMediaSource::GetMediaOption(const FName& Key, bool DefaultValue) const
 	{
 		return bFillGapsInSequence;
 	}
+	if (Key == ImgMedia::ReadVirtualTextureTiles)
+	{
+		return bReadVirtualTextureTiles;
+	}
 
 	return Super::GetMediaOption(Key, DefaultValue);
 }
@@ -143,7 +148,8 @@ bool UImgMediaSource::HasMediaOption(const FName& Key) const
 		(Key == ImgMedia::FrameRateOverrideDenonimatorOption) ||
 		(Key == ImgMedia::FrameRateOverrideNumeratorOption) ||
 		(Key == ImgMedia::ProxyOverrideOption) ||
-		(Key == ImgMedia::MipMapInfoOption))
+		(Key == ImgMedia::MipMapInfoOption) ||
+		(Key == ImgMedia::ReadVirtualTextureTiles))
 	{
 		return true;
 	}
diff --git a/Engine/Plugins/Media/ImgMedia/Source/ImgMedia/Private/ImgMediaPrivate.h b/Engine/Plugins/Media/ImgMedia/Source/ImgMedia/Private/ImgMediaPrivate.h
--- a/Engine/Plugins/Media/ImgMedia/Source/ImgMedia/Private/ImgMediaPrivate.h
+++ b/Engine/Plugins/Media/ImgMedia/Source/ImgMedia/Private/ImgMediaPrivate.h
@@ -54,4 +54,7 @@ namespace ImgMedia
 
 	/** Name of the NumTilesY media option. */
 	static const FName NumTilesYOption("ImgNumTilesY");
+
+	/** Name of the ReadVirtualTextureTiles media option. */
+	static const FName ReadVirtualTextureTiles("ImgReadVirtualTextureTiles");
 }
diff --git a/Engine/Plugins/Media/ImgMedia/Source/ImgMedia/Private/Loader/ImgMediaLoader.cpp b/Engine/Plugins/Media/ImgMedia/Source/ImgMedia/Private/Loader/ImgMediaLoader.cpp
--- a/Engine/Plugins/Media/ImgMedia/Source/ImgMedia/Private/Loader/ImgMediaLoader.cpp
+++ b/Engine/Plugins/Media/ImgMedia/Source/ImgMedia/Private/Loader/ImgMediaLoader.cpp
@@ -121,11 +121,13 @@ namespace ImgMediaLoader
 FImgMediaLoader::FImgMediaLoader(const TSharedRef<FImgMediaScheduler, ESPMode::ThreadSafe>& InScheduler,
 	const TSharedRef<FImgMediaGlobalCache, ESPMode::ThreadSafe>& InGlobalCache,
 	const TSharedPtr<FImgMediaMipMapInfo, ESPMode::ThreadSafe>& InMipMapInfo,
-	bool bInFillGapsInSequence)
+	bool bInFillGapsInSequence,
+	bool bInReadVirtualTextureTiles)
 	: Frames(1)
 	, ImageWrapperModule(FModuleManager::LoadModuleChecked<IImageWrapperModule>("ImageWrapper"))
 	, Initialized(false)
 	, bFillGapsInSequence(bInFillGapsInSequence)
+	, bReadVirtualTextureTiles(bInReadVirtualTextureTiles)
 	, bIsTiled(false)
 	, NumTilesX(0)
 	, NumTilesY(0)
@@ -1345,7 +1347,7 @@ void FImgMediaLoader::AddFrameToCache(int32 FrameNumber, const TSharedPtr<FImgMe
 
 void FImgMediaLoader::GetDesiredMipTiles(int32 FrameIndex, TMap<int32, FImgMediaTileSelection>& OutMipsAndTiles)
 {
-	if(GetDefault<UImgMediaSettings>()->bUseVirtualTextureTiles)
+	if(bReadVirtualTextureTiles)
 	{
 		TMap<int32, TSet<FMediaTileCoordinate>> VisibleTiles;
 		GetVisibleTiles(VisibleTiles);
diff --git a/Engine/Plugins/Media/ImgMedia/Source/ImgMedia/Private/Loader/ImgMediaLoader.h b/Engine/Plugins/Media/ImgMedia/Source/ImgMedia/Private/Loader/ImgMediaLoader.h
--- a/Engine/Plugins/Media/ImgMedia/Source/ImgMedia/Private/Loader/ImgMediaLoader.h
+++ b/Engine/Plugins/Media/ImgMedia/Source/ImgMedia/Private/Loader/ImgMediaLoader.h
@@ -40,7 +40,8 @@ public:
 	FImgMediaLoader(const TSharedRef<FImgMediaScheduler, ESPMode::ThreadSafe>& InScheduler,
 		const TSharedRef<FImgMediaGlobalCache, ESPMode::ThreadSafe>& InGlobalCache,
 		const TSharedPtr<FImgMediaMipMapInfo, ESPMode::ThreadSafe>& InMipMapInfo,
-		bool bInFillGapsInSequence);
+		bool bInFillGapsInSequence,
+		bool bInReadVirtualTextureTiles);
 
 	/** Virtual destructor. */
 	virtual ~FImgMediaLoader();
@@ -454,7 +455,10 @@ private:
 	bool Initialized;
 
 	/** If true, then any gaps in the sequence will be filled with blank frames. */
-	bool bFillGapsInSequence;
+	const bool bFillGapsInSequence;
+
+	/** If true, then virtual texture tiles will be used instead of the MipMapInfo object (applicable for tiled image sequences). */
+	const bool bReadVirtualTextureTiles;
 
 	/** True if this sequence contains tiles. */
 	bool bIsTiled;
diff --git a/Engine/Plugins/Media/ImgMedia/Source/ImgMedia/Private/Player/ImgMediaPlayer.cpp b/Engine/Plugins/Media/ImgMedia/Source/ImgMedia/Private/Player/ImgMediaPlayer.cpp
--- a/Engine/Plugins/Media/ImgMedia/Source/ImgMedia/Private/Player/ImgMediaPlayer.cpp
+++ b/Engine/Plugins/Media/ImgMedia/Source/ImgMedia/Private/Player/ImgMediaPlayer.cpp
@@ -182,11 +182,13 @@ bool FImgMediaPlayer::Open(const FString& Url, const IMediaOptions* Options)
 	FFrameRate FrameRateOverride(0, 0);
 	TSharedPtr<FImgMediaMipMapInfo, ESPMode::ThreadSafe> MipMapInfo;
 	bool bFillGapsInSequence = true;
+	bool bReadVirtualTextureTiles = false;
 	if (Options != nullptr)
 	{
 		FrameRateOverride.Denominator = Options->GetMediaOption(ImgMedia::FrameRateOverrideDenonimatorOption, 0LL);
 		FrameRateOverride.Numerator = Options->GetMediaOption(ImgMedia::FrameRateOverrideNumeratorOption, 0LL);
 		bFillGapsInSequence = Options->GetMediaOption(ImgMedia::FillGapsInSequenceOption, true);
+		bReadVirtualTextureTiles = Options->GetMediaOption(ImgMedia::ReadVirtualTextureTiles, false);
 
 		TSharedPtr<IMediaOptions::FDataContainer, ESPMode::ThreadSafe> DefaultValue;
 		TSharedPtr<IMediaOptions::FDataContainer, ESPMode::ThreadSafe> DataContainer = Options->GetMediaOption(ImgMedia::MipMapInfoOption, DefaultValue);
@@ -209,7 +211,7 @@ bool FImgMediaPlayer::Open(const FString& Url, const IMediaOptions* Options)
 	}
 
 	// initialize image loader on a separate thread
-	Loader = MakeShared<FImgMediaLoader, ESPMode::ThreadSafe>(Scheduler.ToSharedRef(), GlobalCache.ToSharedRef(), MipMapInfo, bFillGapsInSequence);
+	Loader = MakeShared<FImgMediaLoader, ESPMode::ThreadSafe>(Scheduler.ToSharedRef(), GlobalCache.ToSharedRef(), MipMapInfo, bFillGapsInSequence, bReadVirtualTextureTiles);
 	Scheduler->RegisterLoader(Loader.ToSharedRef());
 
 	const FString SequencePath = Url.RightChop(6);
diff --git a/Engine/Plugins/Media/ImgMedia/Source/ImgMedia/Public/ImgMediaSource.h b/Engine/Plugins/Media/ImgMedia/Source/ImgMedia/Public/ImgMediaSource.h
--- a/Engine/Plugins/Media/ImgMedia/Source/ImgMedia/Public/ImgMediaSource.h
+++ b/Engine/Plugins/Media/ImgMedia/Source/ImgMedia/Public/ImgMediaSource.h
@@ -66,6 +66,10 @@ public:
 	UPROPERTY(BlueprintReadWrite, EditAnywhere, Category = Sequence)
 	bool bFillGapsInSequence;
 
+	/** If true, tiled image sequences will rely only on visible tiles provided by the virtual texture system for playback. */
+	UPROPERTY(BlueprintReadWrite, EditAnywhere, Category = Sequence, AdvancedDisplay)
+	bool bReadVirtualTextureTiles;
+
 #if WITH_EDITORONLY_DATA
 
 	/** This is only used so we can customize editing of SequencePath. */
diff --git a/Engine/Plugins/Media/ImgMedia/Source/ImgMediaFactory/Private/ImgMediaSettings.cpp b/Engine/Plugins/Media/ImgMedia/Source/ImgMediaFactory/Private/ImgMediaSettings.cpp
--- a/Engine/Plugins/Media/ImgMedia/Source/ImgMediaFactory/Private/ImgMediaSettings.cpp
+++ b/Engine/Plugins/Media/ImgMedia/Source/ImgMediaFactory/Private/ImgMediaSettings.cpp
@@ -8,7 +8,6 @@
 
 UImgMediaSettings::UImgMediaSettings()
 	: DefaultFrameRate(24, 1)
-	, bUseVirtualTextureTiles(false)
 	, CacheBehindPercentage(25)
 	, CacheSizeGB(1.0f)
 	, CacheThreads(8)
diff --git a/Engine/Plugins/Media/ImgMedia/Source/ImgMediaFactory/Public/ImgMediaSettings.h b/Engine/Plugins/Media/ImgMedia/Source/ImgMediaFactory/Public/ImgMediaSettings.h
--- a/Engine/Plugins/Media/ImgMedia/Source/ImgMediaFactory/Public/ImgMediaSettings.h
+++ b/Engine/Plugins/Media/ImgMedia/Source/ImgMediaFactory/Public/ImgMediaSettings.h
@@ -26,9 +26,6 @@ public:
 	UPROPERTY(config, EditAnywhere, Category=General)
 	FFrameRate DefaultFrameRate;
 
-	/** Whether to use the tiles produced by virtual texturing or not. */
-	UPROPERTY(config, EditAnywhere, Category=General)
-	bool bUseVirtualTextureTiles;
 public:
 
 	/** Percentage of cache to use for frames behind the play head (default = 25%). */
