OpenExrWrapper: Added OpenExr header reader that doesn't have additional overheads that regular exr reader does.

#rnx

diff --git a/Engine/Plugins/Media/ImgMedia/Source/ImgMedia/Private/Readers/ExrImgMediaReader.cpp b/Engine/Plugins/Media/ImgMedia/Source/ImgMedia/Private/Readers/ExrImgMediaReader.cpp
--- a/Engine/Plugins/Media/ImgMedia/Source/ImgMedia/Private/Readers/ExrImgMediaReader.cpp
+++ b/Engine/Plugins/Media/ImgMedia/Source/ImgMedia/Private/Readers/ExrImgMediaReader.cpp
@@ -60,9 +60,7 @@ bool FExrImgMediaReader::GetFrameInfo(const FString& ImagePath, FImgMediaFrameIn
 	// Get tile info.
 	int32 NumTilesX = Loader->GetNumTilesX();
 	int32 NumTilesY = Loader->GetNumTilesY();
-	FRgbaInputFile InputFile(ImagePath, 2);
-
-	if (GetInfo(InputFile, OutInfo))
+	if (GetInfo(ImagePath, OutInfo))
 	{
 		OutInfo.Dim.X *= NumTilesX;
 		OutInfo.Dim.Y *= NumTilesY;
@@ -95,8 +93,7 @@ bool FExrImgMediaReader::ReadFrame(int32 FrameId, const TMap<int32, FImgMediaTil
 	{
 		// Nope. Create it.
 		const FString& LargestImage = Loader->GetImagePath(FrameId, 0);
-		FRgbaInputFile InputFile(LargestImage, 2);
-		if (!GetInfo(InputFile, OutFrame->Info))
+		if (!GetInfo(LargestImage, OutFrame->Info))
 		{
 			return false;
 		}
@@ -254,8 +251,7 @@ TSharedPtr<IImgMediaReader, ESPMode::ThreadSafe> FExrImgMediaReader::GetReader(c
 	}
 	
 	FImgMediaFrameInfo Info;
-
-	if (!GetInfo(InputFile, Info))
+	if (!GetInfo(FirstImageInSequencePath, Info))
 	{
 		return MakeShareable(new FExrImgMediaReader(InLoader));
 	}
@@ -291,40 +287,41 @@ TSharedPtr<IImgMediaReader, ESPMode::ThreadSafe> FExrImgMediaReader::GetReader(c
 /* FExrImgMediaReader implementation
  *****************************************************************************/
 
-bool FExrImgMediaReader::GetInfo(FRgbaInputFile& InputFile, FImgMediaFrameInfo& OutInfo)
+bool FExrImgMediaReader::GetInfo(const FString& FilePath, FImgMediaFrameInfo& OutInfo)
 {
-	if (InputFile.HasInputFile() == false)
+	FOpenExrHeaderReader HeaderReader(FilePath);
+	if (HeaderReader.HasInputFile() == false)
 	{
 		return false;
 	}
 
-	OutInfo.CompressionName = InputFile.GetCompressionName();
-	OutInfo.Dim = InputFile.GetDataWindow();
+	OutInfo.CompressionName = HeaderReader.GetCompressionName();
+	OutInfo.Dim = HeaderReader.GetDataWindow();
 	OutInfo.FormatName = TEXT("EXR");
-	OutInfo.FrameRate = InputFile.GetFrameRate(ImgMedia::DefaultFrameRate);
+	OutInfo.FrameRate = HeaderReader.GetFrameRate(ImgMedia::DefaultFrameRate);
 	OutInfo.Srgb = false;
-	OutInfo.UncompressedSize = InputFile.GetUncompressedSize();
-	OutInfo.NumChannels = InputFile.GetNumChannels();
+	OutInfo.UncompressedSize = HeaderReader.GetUncompressedSize();
+	OutInfo.NumChannels = HeaderReader.GetNumChannels();
 
 	int32 CustomFormat = 0;
-	InputFile.GetIntAttribute(IImgMediaModule::CustomFormatAttributeName.Resolve().ToString(), CustomFormat);
+	HeaderReader.GetIntAttribute(IImgMediaModule::CustomFormatAttributeName.Resolve().ToString(), CustomFormat);
 	bool bIsCustomFormat = CustomFormat > 0;
 	if (bIsCustomFormat)
 	{
 		// Get tile size.
-		OutInfo.bHasTiles = InputFile.GetIntAttribute(IImgMediaModule::CustomFormatTileWidthAttributeName.Resolve().ToString(), OutInfo.TileDimensions.X);
-		OutInfo.bHasTiles = OutInfo.bHasTiles && InputFile.GetIntAttribute(IImgMediaModule::CustomFormatTileHeightAttributeName.Resolve().ToString(), OutInfo.TileDimensions.Y);
-		InputFile.GetIntAttribute(IImgMediaModule::CustomFormatTileBorderAttributeName.Resolve().ToString(), OutInfo.TileBorder);
+		HeaderReader.GetIntAttribute(IImgMediaModule::CustomFormatTileBorderAttributeName.Resolve().ToString(), OutInfo.TileBorder);
+		OutInfo.bHasTiles = HeaderReader.GetIntAttribute(IImgMediaModule::CustomFormatTileWidthAttributeName.Resolve().ToString(), OutInfo.TileDimensions.X);
+		OutInfo.bHasTiles = OutInfo.bHasTiles && HeaderReader.GetIntAttribute(IImgMediaModule::CustomFormatTileHeightAttributeName.Resolve().ToString(), OutInfo.TileDimensions.Y);
 	}
 	else
 	{
-		OutInfo.bHasTiles = InputFile.GetTileSize(OutInfo.TileDimensions);
+		OutInfo.bHasTiles = HeaderReader.GetTileSize(OutInfo.TileDimensions);
 		OutInfo.TileBorder = 0;
 	}
 
 	if (OutInfo.bHasTiles)
 	{
-		OutInfo.NumTiles = FIntPoint(OutInfo.Dim.X / OutInfo.TileDimensions.X, OutInfo.Dim.Y / OutInfo.TileDimensions.Y);
+		OutInfo.NumTiles = FIntPoint(OutInfo.Dim.X / (OutInfo.TileDimensions.X + OutInfo.TileBorder * 2), OutInfo.Dim.Y / (OutInfo.TileDimensions.Y + OutInfo.TileBorder * 2));
 	}
 	else
 	{
@@ -335,6 +332,7 @@ bool FExrImgMediaReader::GetInfo(FRgbaInputFile& InputFile, FImgMediaFrameInfo&
 	return (OutInfo.UncompressedSize > 0) && (OutInfo.Dim.GetMin() > 0);
 }
 
+
 void FExrImgMediaReader::SetCustomFormatInfo(bool bInIsCustomFormat, const FIntPoint& InTileSize)
 {
 	bIsCustomFormat = bInIsCustomFormat;
diff --git a/Engine/Plugins/Media/ImgMedia/Source/ImgMedia/Private/Readers/ExrImgMediaReader.h b/Engine/Plugins/Media/ImgMedia/Source/ImgMedia/Private/Readers/ExrImgMediaReader.h
--- a/Engine/Plugins/Media/ImgMedia/Source/ImgMedia/Private/Readers/ExrImgMediaReader.h
+++ b/Engine/Plugins/Media/ImgMedia/Source/ImgMedia/Private/Readers/ExrImgMediaReader.h
@@ -9,6 +9,7 @@
 #include "IImgMediaReader.h"
 
 class FImgMediaLoader;
+class FOpenExrHeaderReader;
 class FRgbaInputFile;
 
 
@@ -49,7 +50,7 @@ protected:
 	 * @param OutInfo Will contain the frame information.
 	 * @return true on success, false otherwise.
 	 */
-	static bool GetInfo(FRgbaInputFile& InputFile, FImgMediaFrameInfo& OutInfo);
+	static bool GetInfo(const FString& FilePath, FImgMediaFrameInfo& OutInfo);
 
 	/**
 	 * Sets parameters of our custom format images.
diff --git a/Engine/Plugins/Media/ImgMedia/Source/ImgMedia/Private/Readers/ExrImgMediaReaderGpu.cpp b/Engine/Plugins/Media/ImgMedia/Source/ImgMedia/Private/Readers/ExrImgMediaReaderGpu.cpp
--- a/Engine/Plugins/Media/ImgMedia/Source/ImgMedia/Private/Readers/ExrImgMediaReaderGpu.cpp
+++ b/Engine/Plugins/Media/ImgMedia/Source/ImgMedia/Private/Readers/ExrImgMediaReaderGpu.cpp
@@ -28,7 +28,6 @@
 
 #include "ID3D12DynamicRHI.h"
 
-#define READ_IN_CHUNKS 1
 
 DECLARE_GPU_STAT_NAMED(ExrImgMediaReaderGpu, TEXT("ExrImgMediaReaderGpu"));
 
@@ -137,9 +136,8 @@ bool FExrImgMediaReaderGpu::ReadFrame(int32 FrameId, const TMap<int32, FImgMedia
 
 
 	const FString& LargestImagePath = Loader->GetImagePath(FrameId, 0);
-	FRgbaInputFile InputFile(LargestImagePath);
 
-	if (!GetInfo(InputFile, OutFrame->Info))
+	if (!GetInfo(LargestImagePath, OutFrame->Info))
 	{
 		return false;
 	}
@@ -202,8 +200,7 @@ bool FExrImgMediaReaderGpu::ReadFrame(int32 FrameId, const TMap<int32, FImgMedia
 				
 				EReadResult ReadResult = Fail;
 
-#if READ_IN_CHUNKS
-				FRgbaInputFile InputTileFile(ImagePath);
+				FOpenExrHeaderReader InputTileFile(ImagePath);
 				if (InputTileFile.HasInputFile())
 				{
 					// read frame data
@@ -233,36 +230,13 @@ bool FExrImgMediaReaderGpu::ReadFrame(int32 FrameId, const TMap<int32, FImgMedia
 					UE_LOG(LogImgMedia, Error, TEXT("Could not load %s"), *ImagePath);
 					return false;
 				}
-#else
-				FRgbaInputFile InputTileFile(ImagePath);
-				if (InputTileFile.HasInputFile())
-				{
-					// read frame data
-					bool bResult = FExrReader::GenerateTextureData(MipDataPtr, BufferSize, ImagePath, CurrentMipDim.Y, NumChannels);
-					ReadResult = bResult ? Success : Fail;
-					if (bResult)
-					{
-						OutFrame->MipMapsPresent |= 1 << CurrentMipLevel;
-					}
-				}
-				else
-				{
-					UE_LOG(LogImgMedia, Error, TEXT("Could not load %s"), *ImagePath);
-					return false;
-				}
-#endif
 				if (ReadResult == Fail)
 				{
 					// Check if we have a compressed file.
-					FRgbaInputFile InputFileMip(ImagePath);
-					FImgMediaFrameInfo Info;
-					if (GetInfo(InputFileMip, Info))
+					if (OutFrame->Info.CompressionName != "Uncompressed")
 					{
-						if (Info.CompressionName != "Uncompressed")
-						{
-							UE_LOG(LogImgMedia, Error, TEXT("GPU Reader cannot read compressed file %s."), *ImagePath);
-							UE_LOG(LogImgMedia, Error, TEXT("Compressed and uncompressed files should not be mixed in a single sequence."));
-						}
+						UE_LOG(LogImgMedia, Error, TEXT("GPU Reader cannot read compressed file %s."), *ImagePath);
+						UE_LOG(LogImgMedia, Error, TEXT("Compressed and uncompressed files should not be mixed in a single sequence."));
 					}
 
 					// Fall back to CPU.
diff --git a/Engine/Plugins/Media/ImgMedia/Source/OpenExrWrapper/Private/OpenExrWrapperContextReader.cpp b/Engine/Plugins/Media/ImgMedia/Source/OpenExrWrapper/Private/OpenExrWrapperContextReader.cpp
new file mode 100644
--- /dev/null
+++ b/Engine/Plugins/Media/ImgMedia/Source/OpenExrWrapper/Private/OpenExrWrapperContextReader.cpp
@@ -0,0 +1,168 @@
+// Copyright Epic Games, Inc. All Rights Reserved.
+
+#include <exception>
+#include "OpenExrWrapper.h"
+#include "Containers/UnrealString.h"
+#include "Logging/LogMacros.h"
+#include "Modules/ModuleManager.h"
+
+PRAGMA_DEFAULT_VISIBILITY_START
+THIRD_PARTY_INCLUDES_START
+	#include "OpenEXR/openexr_context.h"
+	#include "OpenEXR/openexr_part.h"
+	#include "OpenEXR/ImfCompressionAttribute.h"
+THIRD_PARTY_INCLUDES_END
+PRAGMA_DEFAULT_VISIBILITY_END
+
+DECLARE_LOG_CATEGORY_EXTERN(LogOpenEXRWrapperHeaderReader, Log, All);
+DEFINE_LOG_CATEGORY(LogOpenEXRWrapperHeaderReader);
+
+#define EXR_ATTRIBUTE_COMPRESSION "compression"
+#define EXR_ATTRIBUTE_DATAWINDOW "dataWindow"
+#define EXR_ATTRIBUTE_CHANNELS "channels"
+#define EXR_ATTRIBUTE_TILEDESC "tiles"
+#define EXR_ATTRIBUTE_FPS "framesPerSecond"
+
+/* FOpenExrHeaderReader
+ *****************************************************************************/
+
+namespace
+{
+	bool CheckExrResult(const exr_result_t& Result)
+	{
+		if (Result != EXR_ERR_SUCCESS)
+		{
+			UE_LOG(LogOpenEXRWrapperHeaderReader, Error, TEXT("OpenExr FOpenExrHeaderReader: Issue extracting data via C Interface."));
+		}
+		return Result == EXR_ERR_SUCCESS;
+	}
+}
+
+FOpenExrHeaderReader::FOpenExrHeaderReader(const FString& FilePath) : FileContext(new exr_context_t())
+{
+	exr_context_initializer_t ContextInitializer = EXR_DEFAULT_CONTEXT_INITIALIZER;
+	if (!CheckExrResult(exr_start_read((exr_context_t*)FileContext.Get(), TCHAR_TO_ANSI(*FilePath), &ContextInitializer)))
+	{
+		FileContext.Reset();
+		return;
+	}
+}
+
+FOpenExrHeaderReader::~FOpenExrHeaderReader()
+{
+	CheckExrResult(exr_finish((exr_context_t*)FileContext.Get()));
+	FileContext.Reset();
+}
+
+const TCHAR* FOpenExrHeaderReader::GetCompressionName() const
+{
+	const exr_attribute_t* Attribute = nullptr;
+	CheckExrResult(exr_get_attribute_by_name(*((exr_context_t*)FileContext.Get()), 0, EXR_ATTRIBUTE_COMPRESSION, &Attribute));
+
+	if (Attribute == nullptr)
+	{
+		return TEXT("");
+	}
+
+	switch (Attribute->i)
+	{
+	case Imf::NO_COMPRESSION:
+		return TEXT("Uncompressed");
+
+	case Imf::RLE_COMPRESSION:
+		return TEXT("RLE");
+
+	case Imf::ZIPS_COMPRESSION:
+		return TEXT("ZIPS");
+
+	case Imf::ZIP_COMPRESSION:
+		return TEXT("ZIP");
+
+	case Imf::PIZ_COMPRESSION:
+		return TEXT("PIZ");
+
+	case Imf::PXR24_COMPRESSION:
+		return TEXT("PXR24");
+
+	case Imf::B44_COMPRESSION:
+		return TEXT("B44");
+
+	case Imf::B44A_COMPRESSION:
+		return TEXT("B44A");
+
+	default:
+		return TEXT("Unknown");
+	}
+
+	return nullptr;
+}
+
+FIntPoint FOpenExrHeaderReader::GetDataWindow() const
+{
+	const exr_attribute_t* Attribute = nullptr;
+	CheckExrResult(exr_get_attribute_by_name(*((exr_context_t*)FileContext.Get()), 0, EXR_ATTRIBUTE_DATAWINDOW, &Attribute));
+	return FIntPoint(
+		Attribute->box2i->max.x - Attribute->box2i->min.x + 1,
+		Attribute->box2i->max.y - Attribute->box2i->min.y + 1
+	);
+}
+
+FFrameRate FOpenExrHeaderReader::GetFrameRate(const FFrameRate& DefaultValue) const
+{
+	const exr_attribute_t* Attribute = nullptr;
+	exr_result_t Result = exr_get_attribute_by_name(*((exr_context_t*)FileContext.Get()), 0, EXR_ATTRIBUTE_FPS, &Attribute);
+
+	if (Result == EXR_ERR_NO_ATTR_BY_NAME || Attribute == nullptr)
+	{
+		return DefaultValue;
+	}
+
+	return FFrameRate(Attribute->rational->num, Attribute->rational->denom);
+}
+
+int32 FOpenExrHeaderReader::GetUncompressedSize() const
+{
+	const int32 NumChannels = GetNumChannels();
+	const int32 ChannelSize = sizeof(int16);
+	const FIntPoint Window = GetDataWindow();
+
+	return (Window.X * Window.Y * NumChannels * ChannelSize);
+}
+
+int32 FOpenExrHeaderReader::GetNumChannels() const
+{
+	const exr_attribute_t* Attribute = nullptr;
+	CheckExrResult(exr_get_attribute_by_name(*((exr_context_t*)FileContext.Get()), 0, EXR_ATTRIBUTE_CHANNELS, &Attribute));
+	return Attribute->chlist->num_channels;
+}
+
+bool FOpenExrHeaderReader::GetTileSize(FIntPoint& OutTileSize) const
+{
+	const exr_attribute_t* Attribute = nullptr;
+	exr_result_t Result = exr_get_attribute_by_name(*((exr_context_t*)FileContext.Get()), 0, EXR_ATTRIBUTE_TILEDESC, &Attribute);
+	if (Result == EXR_ERR_SUCCESS && Attribute)
+	{
+		OutTileSize = FIntPoint(Attribute->tiledesc->x_size, Attribute->tiledesc->y_size);
+	}
+	return Attribute != nullptr;
+}
+
+bool FOpenExrHeaderReader::HasInputFile() const
+{
+	return FileContext != nullptr;
+}
+
+bool FOpenExrHeaderReader::GetIntAttribute(const FString& Name, int32& Value)
+{
+	const exr_attribute_t* Attribute = nullptr;
+	exr_result_t Result = exr_get_attribute_by_name(*((exr_context_t*)FileContext.Get()), 0, TCHAR_TO_ANSI(*Name), &Attribute);
+	if (Result == EXR_ERR_SUCCESS && Attribute)
+	{
+		Value = Attribute->i;
+		return true;
+	}
+	return false;
+}
+
+
+
diff --git a/Engine/Plugins/Media/ImgMedia/Source/OpenExrWrapper/Public/OpenExrWrapper.h b/Engine/Plugins/Media/ImgMedia/Source/OpenExrWrapper/Public/OpenExrWrapper.h
--- a/Engine/Plugins/Media/ImgMedia/Source/OpenExrWrapper/Public/OpenExrWrapper.h
+++ b/Engine/Plugins/Media/ImgMedia/Source/OpenExrWrapper/Public/OpenExrWrapper.h
@@ -56,6 +56,40 @@ private:
 	void* InputFile;
 };
 
+class OPENEXRWRAPPER_API FOpenExrHeaderReader
+{
+public:
+
+	FOpenExrHeaderReader(const FString& FilePath);
+	~FOpenExrHeaderReader();
+
+public:
+
+	const TCHAR* GetCompressionName() const;
+	FIntPoint GetDataWindow() const;
+	FFrameRate GetFrameRate(const FFrameRate& DefaultValue) const;
+	int32 GetUncompressedSize() const;
+	int32 GetNumChannels() const;
+
+	// Gets tile dimensions. Returns false if image has no tiles.
+	bool GetTileSize(FIntPoint& OutTileSize) const;
+
+	bool HasInputFile() const;
+
+	/**
+	 * Get an attribute from the image.
+	 *
+	 * @param Name		Name of attribute.
+	 * @param Value		Will be set to the value of the attribute if the attribute is found.
+	 *					Will NOT be set if the attribute is not found.
+	 * @return			True if the attribute was found, false otherwise.
+	 */
+	bool GetIntAttribute(const FString& Name, int32& Value);
+
+private:
+	TSharedPtr<void> FileContext;
+};
+
 /**
  * Base class for our classes that output EXR files.
  */
