ImgMedia: Added support to process images for RGB.

#rnx

diff --git a/Engine/Plugins/Media/ImgMedia/Source/ImgMediaEditor/Private/Widgets/SImgMediaProcessImages.cpp b/Engine/Plugins/Media/ImgMedia/Source/ImgMediaEditor/Private/Widgets/SImgMediaProcessImages.cpp
--- a/Engine/Plugins/Media/ImgMedia/Source/ImgMediaEditor/Private/Widgets/SImgMediaProcessImages.cpp
+++ b/Engine/Plugins/Media/ImgMedia/Source/ImgMediaEditor/Private/Widgets/SImgMediaProcessImages.cpp
@@ -141,6 +141,9 @@ void SImgMediaProcessImages::ProcessAllImages(TSharedPtr<SNotificationItem> Conf
 			IImageWrapperModule& ImageWrapperModule = FModuleManager::LoadModuleChecked<IImageWrapperModule>("ImageWrapper");
 			TSharedPtr<IImageWrapper> ImageWrapper = ImageWrapperModule.CreateImageWrapper(ImageFormat);
 
+			// ImageWrapper is always returning an alpha channel for RGB, so check if we really have one.
+			bool bHasAlphaChannel = HasAlphaChannel(Ext, FPaths::Combine(SequencePath, FoundFiles[0]));
+
 			// Loop through all files.
 			int NumDone = 0;
 			int TotalNum = FoundFiles.Num();
@@ -178,7 +181,7 @@ void SImgMediaProcessImages::ProcessAllImages(TSharedPtr<SNotificationItem> Conf
 				if (bUseCustomFormat)
 				{
 					ProcessImageCustom(ImageWrapper, InTileWidth, InTileHeight, TileBorder,
-						bEnableMips, Name);
+						bEnableMips, bHasAlphaChannel, Name);
 				}
 				else
 				{
@@ -201,6 +204,19 @@ void SImgMediaProcessImages::ProcessAllImages(TSharedPtr<SNotificationItem> Conf
 	});
 }
 
+bool SImgMediaProcessImages::HasAlphaChannel(const FString& Ext, const FString& File)
+{
+	bool bHasAlpha = true;
+	// We just support EXR at the moment.
+	if (Ext == TEXT("exr"))
+	{
+		FRgbaInputFile InputFile(File);
+		bHasAlpha = InputFile.GetNumChannels() == 4;
+	}
+
+	return bHasAlpha;
+}
+
 void SImgMediaProcessImages::ProcessImage(TSharedPtr<IImageWrapper>& InImageWrapper,
 	int32 InTileWidth, int32 InTileHeight, const FString& InName, const FString& FileExtension)
 {
@@ -267,7 +283,7 @@ void SImgMediaProcessImages::ProcessImage(TSharedPtr<IImageWrapper>& InImageWrap
 
 void SImgMediaProcessImages::ProcessImageCustom(TSharedPtr<IImageWrapper>& InImageWrapper,
 	int32 InTileWidth, int32 InTileHeight, int32 InTileBorder, bool bInEnableMips,
-	const FString& InName)
+	bool bHasAlphaChannel, const FString& InName)
 {
 #if IMGMEDIAEDITOR_EXR_SUPPORTED_PLATFORM
 	// Get image data.
@@ -285,6 +301,14 @@ void SImgMediaProcessImages::ProcessImageCustom(TSharedPtr<IImageWrapper>& InIma
 	int32 TileWidth = Width / NumTilesX;
 	int32 TileHeight = Height / NumTilesY;
 	int32 BytesPerPixel = RawData.Num() / (Width * Height);
+	int32 BytesPerPixelPerChannel = BitDepth / 8;
+	int32 NumChannels = BytesPerPixel / BytesPerPixelPerChannel;
+	int32 DestNumChannels = NumChannels;
+	// ImageWrapper always returns an alpha channel, so make sure we really have one.
+	if ((DestNumChannels == 4) && (bHasAlphaChannel == false))
+	{
+		DestNumChannels = 3;
+	}
 
 	bool bIsTiled = (NumTilesX > 1) || (NumTilesY > 1);
 	uint8* RawDataPtr = RawData.GetData();
@@ -375,8 +399,6 @@ void SImgMediaProcessImages::ProcessImageCustom(TSharedPtr<IImageWrapper>& InIma
 	const FString BChannelName = FString(TEXT("B"));
 	const FString AChannelName = FString(TEXT("A"));
 
-	int32 NumChannels = 4;
-
 	FIntPoint Stride(2, 0);
 
 	// Create tiled exr file.
@@ -395,17 +417,29 @@ void SImgMediaProcessImages::ProcessImageCustom(TSharedPtr<IImageWrapper>& InIma
 	}
 
 	// Add channels.
-	OutFile.AddChannel(AChannelName);
-	OutFile.AddChannel(BChannelName);
-	OutFile.AddChannel(GChannelName);
-	OutFile.AddChannel(RChannelName);
+	if (DestNumChannels == 4)
+	{
+		OutFile.AddChannel(AChannelName);
+	}
+	if (DestNumChannels >= 3)
+	{
+		OutFile.AddChannel(BChannelName);
+		OutFile.AddChannel(GChannelName);
+		OutFile.AddChannel(RChannelName);
+	}
 
 	// Create output.
 	OutFile.CreateOutputFile(InName, DestWidth, DestHeight, bInEnableMips);
-	OutFile.AddFrameBufferChannel(AChannelName, nullptr, Stride);
-	OutFile.AddFrameBufferChannel(BChannelName, nullptr, Stride);
-	OutFile.AddFrameBufferChannel(GChannelName, nullptr, Stride);
-	OutFile.AddFrameBufferChannel(RChannelName, nullptr, Stride);
+	if (DestNumChannels == 4)
+	{
+		OutFile.AddFrameBufferChannel(AChannelName, nullptr, Stride);
+	}
+	if (DestNumChannels >= 3)
+	{
+		OutFile.AddFrameBufferChannel(BChannelName, nullptr, Stride);
+		OutFile.AddFrameBufferChannel(GChannelName, nullptr, Stride);
+		OutFile.AddFrameBufferChannel(RChannelName, nullptr, Stride);
+	}
 
 	// Flip between 2 buffers making mips.
 	TArray64<uint8> RawData2;
@@ -466,10 +500,18 @@ void SImgMediaProcessImages::ProcessImageCustom(TSharedPtr<IImageWrapper>& InIma
 
 		// Write to EXR.
 		Stride.Y = MipWidth * BytesPerPixel;
-		OutFile.UpdateFrameBufferChannel(AChannelName, CurrentBuffer, Stride);
-		OutFile.UpdateFrameBufferChannel(BChannelName, CurrentBuffer + MipWidth * 2, Stride);
-		OutFile.UpdateFrameBufferChannel(GChannelName, CurrentBuffer + MipWidth * 4, Stride);
-		OutFile.UpdateFrameBufferChannel(RChannelName, CurrentBuffer + MipWidth * 6, Stride);
+		
+		if (DestNumChannels == 4)
+		{
+			OutFile.UpdateFrameBufferChannel(AChannelName, CurrentBuffer, Stride);
+		}
+
+		OutFile.UpdateFrameBufferChannel(BChannelName,
+			CurrentBuffer + MipWidth * BytesPerPixelPerChannel, Stride);
+		OutFile.UpdateFrameBufferChannel(GChannelName,
+			CurrentBuffer + MipWidth * 2 * BytesPerPixelPerChannel, Stride);
+		OutFile.UpdateFrameBufferChannel(RChannelName,
+			CurrentBuffer + MipWidth * 3 * BytesPerPixelPerChannel, Stride);
 
 		OutFile.SetFrameBuffer();
 
diff --git a/Engine/Plugins/Media/ImgMedia/Source/ImgMediaEditor/Private/Widgets/SImgMediaProcessImages.h b/Engine/Plugins/Media/ImgMedia/Source/ImgMediaEditor/Private/Widgets/SImgMediaProcessImages.h
--- a/Engine/Plugins/Media/ImgMedia/Source/ImgMediaEditor/Private/Widgets/SImgMediaProcessImages.h
+++ b/Engine/Plugins/Media/ImgMedia/Source/ImgMediaEditor/Private/Widgets/SImgMediaProcessImages.h
@@ -36,6 +36,15 @@ private:
 	 */
 	void ProcessAllImages(TSharedPtr<SNotificationItem> ConfirmNotification);
 	
+	/**
+	  * Checks if this file has an alpha channel.
+	  *
+	  * @param Ext			File extension.
+	  * @param File			Full path to file.
+	  * @return				True if the file has an alpha channel.
+	  */
+	bool HasAlphaChannel(const FString& Ext, const FString& File);
+
 	/**
 	 * Processess a single image and writes out 1 or more files.
 	 * Tiles and mips may be generated.
@@ -60,11 +69,12 @@ private:
 	 * @param InTileHeight		Desired height of tiles.
 	 * @param InTileBorder		Number of pixels to duplicate along a tile edge.
 	 * @param bInEnableMips		Turn on mip mapping.
+	 * @param bHasAlphaChannel	True if there really is an alpha channel.
 	 * @param InName			Full path and name of file to write.
 	 */
 	void ProcessImageCustom(TSharedPtr<IImageWrapper>& InImageWrapper,
 		int32 InTileWidth, int32 InTileHeight, int32 InTileBorder, 
-		bool bInEnableMips, const FString& InName);
+		bool bInEnableMips, bool bHasAlphaChannel, const FString& InName);
 
 	/** Holds our details view. */
 	TSharedPtr<class IDetailsView> DetailsView;
