ImgMedia: Added smart caching.

#rnx

diff --git a/Engine/Plugins/Media/ImgMedia/Source/ImgMedia/Private/ImgMediaPrivate.h b/Engine/Plugins/Media/ImgMedia/Source/ImgMedia/Private/ImgMediaPrivate.h
--- a/Engine/Plugins/Media/ImgMedia/Source/ImgMedia/Private/ImgMediaPrivate.h
+++ b/Engine/Plugins/Media/ImgMedia/Source/ImgMedia/Private/ImgMediaPrivate.h
@@ -57,4 +57,10 @@ namespace ImgMedia
 
 	/** Name of the ReadVirtualTextureTiles media option. */
 	static const FName ReadVirtualTextureTiles("ImgReadVirtualTextureTiles");
+
+	/** Name of the smart cache media option. */
+	static const FName SmartCacheEnabled("ImgMediaSmartCacheEnabled");
+
+	/** Name of the smart cache time to look ahead media option. */
+	static const FName SmartCacheTimeToLookAhead("ImgMediaSmartCacheTimeToLookAhead");
 }
diff --git a/Engine/Plugins/Media/ImgMedia/Source/ImgMedia/Private/Loader/ImgMediaLoader.cpp b/Engine/Plugins/Media/ImgMedia/Source/ImgMedia/Private/Loader/ImgMediaLoader.cpp
--- a/Engine/Plugins/Media/ImgMedia/Source/ImgMedia/Private/Loader/ImgMediaLoader.cpp
+++ b/Engine/Plugins/Media/ImgMedia/Source/ImgMedia/Private/Loader/ImgMediaLoader.cpp
@@ -122,7 +122,8 @@ FImgMediaLoader::FImgMediaLoader(const TSharedRef<FImgMediaScheduler, ESPMode::T
 	const TSharedRef<FImgMediaGlobalCache, ESPMode::ThreadSafe>& InGlobalCache,
 	const TSharedPtr<FImgMediaMipMapInfo, ESPMode::ThreadSafe>& InMipMapInfo,
 	bool bInFillGapsInSequence,
-	bool bInReadVirtualTextureTiles)
+	bool bInReadVirtualTextureTiles,
+	const FImgMediaLoaderSmartCacheSettings& InSmartCacheSettings)
 	: Frames(1)
 	, ImageWrapperModule(FModuleManager::LoadModuleChecked<IImageWrapperModule>("ImageWrapper"))
 	, Initialized(false)
@@ -142,6 +143,7 @@ FImgMediaLoader::FImgMediaLoader(const TSharedRef<FImgMediaScheduler, ESPMode::T
 	, LastRequestedFrame(INDEX_NONE)
 	, RetryCount(0)
 	, UseGlobalCache(false)
+	, SmartCacheSettings(InSmartCacheSettings)
 {
 	ResetFetchLogic();
 	UE_LOG(LogImgMedia, Verbose, TEXT("Loader %p: Created"), this);
@@ -817,6 +819,11 @@ void FImgMediaLoader::LoadSequence(const FString& SequencePath, const FFrameRate
 
 	const UImgMediaSettings* Settings = GetDefault<UImgMediaSettings>();
 	UseGlobalCache = Settings->UseGlobalCache;
+	if (SmartCacheSettings.bIsEnabled)
+	{
+		// Smart cache does not use the global cache.
+		UseGlobalCache = false;
+	}
 	SequenceName = FName(*SequencePath);
 
 	// fetch sequence attributes from first image
@@ -895,15 +902,30 @@ void FImgMediaLoader::LoadSequence(const FString& SequencePath, const FFrameRate
 
 	// initialize loader
 	const FPlatformMemoryStats Stats = FPlatformMemory::GetStats();
-	const SIZE_T DesiredCacheSize = Settings->CacheSizeGB * 1024 * 1024 * 1024;
-	const SIZE_T CacheSize = FMath::Clamp(DesiredCacheSize, (SIZE_T)0, (SIZE_T)Stats.AvailablePhysical);
+	
+	// Are we using the smart cache?
+	if (SmartCacheSettings.bIsEnabled)
+	{
+		// Load frames that are a certain amount of time ahead of us.
+		float TimeToLookAhead = SmartCacheSettings.TimeToLookAhead;
+		
+		NumFramesToLoad = TimeToLookAhead / SequenceFrameRate.AsInterval();
+		NumFramesToLoad = FMath::Clamp(NumFramesToLoad, 0, GetNumImages());
+		NumLoadBehind = 0;
+		NumLoadAhead = NumFramesToLoad;
+	}
+	else
+	{
+		const SIZE_T DesiredCacheSize = Settings->CacheSizeGB * 1024 * 1024 * 1024;
+		const SIZE_T CacheSize = FMath::Clamp(DesiredCacheSize, (SIZE_T)0, (SIZE_T)Stats.AvailablePhysical);
 
-	const int32 MaxFramesToLoad = (int32)(CacheSize / UncompressedSize);
-	NumFramesToLoad = FMath::Clamp(MaxFramesToLoad, 0, GetNumImages());
-	const float LoadBehindScale = FMath::Clamp(Settings->CacheBehindPercentage, 0.0f, 100.0f) / 100.0f;
+		const int32 MaxFramesToLoad = (int32)(CacheSize / UncompressedSize);
+		NumFramesToLoad = FMath::Clamp(MaxFramesToLoad, 0, GetNumImages());
+		const float LoadBehindScale = FMath::Clamp(Settings->CacheBehindPercentage, 0.0f, 100.0f) / 100.0f;
 
-	NumLoadBehind = (int32)(LoadBehindScale * MaxFramesToLoad);
-	NumLoadAhead = (int32)((1.0f - LoadBehindScale) * MaxFramesToLoad);
+		NumLoadBehind = (int32)(LoadBehindScale * MaxFramesToLoad);
+		NumLoadAhead = (int32)((1.0f - LoadBehindScale) * MaxFramesToLoad);
+	}
 
 	// Giving our reader a chance to handle RAM allocation.
 	// Not all readers use this, only those that need to handle large files 
diff --git a/Engine/Plugins/Media/ImgMedia/Source/ImgMedia/Private/Loader/ImgMediaLoader.h b/Engine/Plugins/Media/ImgMedia/Source/ImgMedia/Private/Loader/ImgMediaLoader.h
--- a/Engine/Plugins/Media/ImgMedia/Source/ImgMedia/Private/Loader/ImgMediaLoader.h
+++ b/Engine/Plugins/Media/ImgMedia/Source/ImgMedia/Private/Loader/ImgMediaLoader.h
@@ -23,6 +23,23 @@ class IQueuedWork;
 
 struct FImgMediaFrame;
 
+/**
+ * Settings for the smart cache.
+ */
+struct FImgMediaLoaderSmartCacheSettings
+{
+	/** True if we are using the smart cache. */
+	bool bIsEnabled;
+	/** The cache will fill up with frames that are up to this time from the current time. */
+	float TimeToLookAhead;
+
+	/** Constructor. */
+	FImgMediaLoaderSmartCacheSettings(bool bInIsEnabled, float InTimeToLookAhead)
+		: bIsEnabled(bInIsEnabled)
+		, TimeToLookAhead(InTimeToLookAhead)
+	{
+	}
+};
 
 /**
  * Loads image sequence frames from disk.
@@ -41,7 +58,8 @@ public:
 		const TSharedRef<FImgMediaGlobalCache, ESPMode::ThreadSafe>& InGlobalCache,
 		const TSharedPtr<FImgMediaMipMapInfo, ESPMode::ThreadSafe>& InMipMapInfo,
 		bool bInFillGapsInSequence,
-		bool bInReadVirtualTextureTiles);
+		bool bInReadVirtualTextureTiles,
+		const FImgMediaLoaderSmartCacheSettings& InSmartCacheSettings);
 
 	/** Virtual destructor. */
 	virtual ~FImgMediaLoader();
@@ -534,4 +552,7 @@ private:
 
 	/** List of visible tiles updated through the IMediaView interface. */
 	TMap<int32, TSet<FMediaTileCoordinate>> ActiveTilesPerMipLevel;
+
+	/** Settings for the smart cache. */
+	FImgMediaLoaderSmartCacheSettings SmartCacheSettings;
 };
diff --git a/Engine/Plugins/Media/ImgMedia/Source/ImgMedia/Private/Player/ImgMediaPlayer.cpp b/Engine/Plugins/Media/ImgMedia/Source/ImgMedia/Private/Player/ImgMediaPlayer.cpp
--- a/Engine/Plugins/Media/ImgMedia/Source/ImgMedia/Private/Player/ImgMediaPlayer.cpp
+++ b/Engine/Plugins/Media/ImgMedia/Source/ImgMedia/Private/Player/ImgMediaPlayer.cpp
@@ -183,13 +183,16 @@ bool FImgMediaPlayer::Open(const FString& Url, const IMediaOptions* Options)
 	TSharedPtr<FImgMediaMipMapInfo, ESPMode::ThreadSafe> MipMapInfo;
 	bool bFillGapsInSequence = true;
 	bool bReadVirtualTextureTiles = false;
+	bool bIsSmartCacheEnabled = false;
+	float SmartCacheTimeToLookAhead = 0.0f;
 	if (Options != nullptr)
 	{
 		FrameRateOverride.Denominator = Options->GetMediaOption(ImgMedia::FrameRateOverrideDenonimatorOption, 0LL);
 		FrameRateOverride.Numerator = Options->GetMediaOption(ImgMedia::FrameRateOverrideNumeratorOption, 0LL);
 		bFillGapsInSequence = Options->GetMediaOption(ImgMedia::FillGapsInSequenceOption, true);
 		bReadVirtualTextureTiles = Options->GetMediaOption(ImgMedia::ReadVirtualTextureTiles, false);
-
+		bIsSmartCacheEnabled = Options->GetMediaOption(ImgMedia::SmartCacheEnabled, false);
+		SmartCacheTimeToLookAhead = Options->GetMediaOption(ImgMedia::SmartCacheTimeToLookAhead, 0.0f);
 		TSharedPtr<IMediaOptions::FDataContainer, ESPMode::ThreadSafe> DefaultValue;
 		TSharedPtr<IMediaOptions::FDataContainer, ESPMode::ThreadSafe> DataContainer = Options->GetMediaOption(ImgMedia::MipMapInfoOption, DefaultValue);
 		if (DataContainer.IsValid())
@@ -211,7 +214,10 @@ bool FImgMediaPlayer::Open(const FString& Url, const IMediaOptions* Options)
 	}
 
 	// initialize image loader on a separate thread
-	Loader = MakeShared<FImgMediaLoader, ESPMode::ThreadSafe>(Scheduler.ToSharedRef(), GlobalCache.ToSharedRef(), MipMapInfo, bFillGapsInSequence, bReadVirtualTextureTiles);
+	FImgMediaLoaderSmartCacheSettings SmartCacheSettings(bIsSmartCacheEnabled, SmartCacheTimeToLookAhead);
+	Loader = MakeShared<FImgMediaLoader, ESPMode::ThreadSafe>(Scheduler.ToSharedRef(),
+		GlobalCache.ToSharedRef(), MipMapInfo, bFillGapsInSequence, bReadVirtualTextureTiles,
+		SmartCacheSettings);
 	Scheduler->RegisterLoader(Loader.ToSharedRef());
 
 	const FString SequencePath = Url.RightChop(6);
