ImgMedia: Process images can now output tiled, mipmapped, bordered images.

#rnx

diff --git a/Engine/Plugins/Media/ImgMedia/Source/ImgMediaEditor/Private/Widgets/SImgMediaProcessImages.cpp b/Engine/Plugins/Media/ImgMedia/Source/ImgMediaEditor/Private/Widgets/SImgMediaProcessImages.cpp
--- a/Engine/Plugins/Media/ImgMedia/Source/ImgMediaEditor/Private/Widgets/SImgMediaProcessImages.cpp
+++ b/Engine/Plugins/Media/ImgMedia/Source/ImgMediaEditor/Private/Widgets/SImgMediaProcessImages.cpp
@@ -310,89 +310,17 @@ void SImgMediaProcessImages::ProcessImageCustom(TSharedPtr<IImageWrapper>& InIma
 		DestNumChannels = 3;
 	}
 
+	TArray64<uint8> TileBuffer;
 	bool bIsTiled = (NumTilesX > 1) || (NumTilesY > 1);
-	uint8* RawDataPtr = RawData.GetData();
-
-	// Tile data.
-	TArray64<uint8> RawDataTiled;
 	if (bIsTiled)
 	{
-		// We don't support tile borders larger than a tile size,
-		// but this shuld not happen in practice.
-		if ((InTileBorder > TileWidth) || (InTileBorder > TileHeight))
-		{
-			UE_LOG(LogImgMediaEditor, Error, TEXT("Tile border is larger than tile size. Clamping to tile size."));
-			InTileBorder = FMath::Min(TileWidth, TileHeight);
-		}
-
+		// Take border into account.
 		DestWidth = Width + InTileBorder * 2 * NumTilesX;
 		DestHeight = Height + InTileBorder * 2 * NumTilesY;
-		RawDataTiled.AddUninitialized(DestWidth * DestHeight * BytesPerPixel);
-		RawDataPtr = RawDataTiled.GetData();
-
-		uint8* SourceData = RawData.GetData();
-		uint8* DestData = RawDataTiled.GetData();
-		int32 DestTileWidth = TileWidth + InTileBorder * 2;
-		int32 DestTileHeight = TileHeight + InTileBorder * 2;
-		int32 BytesPerTile = TileWidth * TileHeight * BytesPerPixel;
-		int32 ByterPerDestTile = DestTileWidth * DestTileHeight * BytesPerPixel;
-
-		// Loop over y tiles.
-		for (int32 TileY = 0; TileY < NumTilesY; ++TileY)
-		{
-			// Loop over x tiles.
-			for (int32 TileX = 0; TileX < NumTilesX; ++TileX)
-			{
-				// Get address of the source and destination tiles.
-				uint8* SourceTile = SourceData +
-					(TileX * TileWidth + TileY * Width * TileHeight) * BytesPerPixel;
-				uint8* DestTile = DestData + (TileX + TileY * NumTilesX) * ByterPerDestTile;
-				
-				int32 NumberOfPixelsToCopy = TileWidth;
-
-				// Create a left border.
-				if (TileX > 0)
-				{
-					NumberOfPixelsToCopy += InTileBorder;
-					// Offset the source to get the extra pixels.
-					SourceTile -= InTileBorder * BytesPerPixel;
-				}
-				else
-				{
-					// Offset the destination as we are skipping this border as we have no data.
-					DestTile += InTileBorder * BytesPerPixel;
-				}
-
-				// Create a right border.
-				if (TileX < NumTilesX - 1)
-				{
-					NumberOfPixelsToCopy += InTileBorder;
-				}
-				
-				// Loop over each row in the tile.
-				for (int32 Row = 0; Row < DestTileHeight; ++Row)
-				{
-					// Make sure we don't go beyond the source data.
-					int32 SourceRow = Row - InTileBorder;
-					if (TileY == 0)
-					{
-						SourceRow = FMath::Max(SourceRow, 0);
-					}
-					if (TileY == NumTilesY - 1)
-					{
-						SourceRow = FMath::Min(SourceRow, TileHeight - 1);
-					}
-
-					uint8* SourceLine = SourceTile + SourceRow * Width * BytesPerPixel;
-					uint8* DestLine = DestTile + Row * DestTileWidth * BytesPerPixel;
-					
-					// Copy the main data.
-					FMemory::Memcpy(DestLine, SourceLine, NumberOfPixelsToCopy * BytesPerPixel);
-				}
-			}
-		}
 	}
 
+	uint8* RawDataPtr = RawData.GetData();
+
 	// Names for our channels.
 	const FString RChannelName = FString(TEXT("R"));
 	const FString GChannelName = FString(TEXT("G"));
@@ -450,6 +378,8 @@ void SImgMediaProcessImages::ProcessImageCustom(TSharedPtr<IImageWrapper>& InIma
 
 	// Loop over each mip level.
 	int32 NumMips = OutFile.GetNumberOfMipLevels();
+	int32 MipSourceWidth = Width;
+	int32 MipSourceHeight = Height;
 	for (int32 MipLevel = 0; MipLevel < NumMips; MipLevel++)
 	{
 		int32 MipWidth = OutFile.GetMipWidth(MipLevel);
@@ -498,6 +428,55 @@ void SImgMediaProcessImages::ProcessImageCustom(TSharedPtr<IImageWrapper>& InIma
 			}
 		}
 
+		// Do we need to tile this mip?
+		// Need to also check that this is actually a valid mip level.
+		if ((bIsTiled) && (MipSourceWidth > 0) && (MipSourceHeight > 0))
+		{
+			int32 MipTileWidth = TileWidth;
+			int32 MipTileHeight = TileHeight;
+
+			// A tile could be larger than the mip level when dealing with mips.
+			if (MipTileWidth > MipSourceWidth)
+			{
+				MipTileWidth = MipSourceWidth;
+			}
+			if (MipTileHeight > MipSourceHeight)
+			{
+				MipTileHeight = MipSourceHeight;
+			}
+
+			int32 OutputWidth = 0;
+			int32 OutputHeight = 0;
+			int32 MipNumTilesX = MipSourceWidth / MipTileWidth;
+			int32 MipNumTilesY = MipSourceHeight / MipTileHeight;
+
+			// Make sure our sizes match the mip size we get from EXR.
+			int32 ExpectedMipWidth = MipSourceWidth + MipNumTilesX * InTileBorder * 2;
+			if (ExpectedMipWidth != MipWidth)
+			{
+				UE_LOG(LogImgMediaEditor, Error,
+					TEXT("Expected mip level width of %d, but got %d (SourceWidth:%d NumTiles:%d TileBorder:%d"),
+					ExpectedMipWidth, MipHeight,
+					MipSourceWidth, MipNumTilesX, InTileBorder);
+			}
+			int32 ExpectedMipHeight = MipSourceHeight + MipNumTilesY * InTileBorder * 2;
+			if (ExpectedMipHeight != MipHeight)
+			{
+				UE_LOG(LogImgMediaEditor, Error,
+					TEXT("Expected mip level height of %d, but got %d (SourceHeight:%d NumTiles:%d TileBorder:%d"),
+					ExpectedMipHeight, MipHeight,
+					MipSourceHeight, MipNumTilesY, InTileBorder);
+			}
+
+			// Tile the buffer.
+			TileData(CurrentBuffer, TileBuffer,
+				MipSourceWidth, MipSourceHeight, MipWidth, MipHeight,
+				MipNumTilesX, MipNumTilesY,
+				MipTileWidth, MipTileHeight, InTileBorder,
+				BytesPerPixel);
+			CurrentBuffer = TileBuffer.GetData();
+		}
+
 		// Write to EXR.
 		Stride.Y = MipWidth * BytesPerPixel;
 		
@@ -519,6 +498,8 @@ void SImgMediaProcessImages::ProcessImageCustom(TSharedPtr<IImageWrapper>& InIma
 
 		// Switch buffers.
 		CurrentMipBufferIndex ^= 1;
+		MipSourceHeight /= 2;
+		MipSourceWidth /= 2;
 	}
 
 #else // IMGMEDIAEDITOR_EXR_SUPPORTED_PLATFORM
@@ -526,4 +507,92 @@ void SImgMediaProcessImages::ProcessImageCustom(TSharedPtr<IImageWrapper>& InIma
 #endif // IMGMEDIAEDITOR_EXR_SUPPORTED_PLATFORM
 }
 
+void SImgMediaProcessImages::TileData(uint8* SourceData, TArray64<uint8>& DestArray,
+	int32 SourceWidth, int32 SourceHeight, int32 DestWidth, int32 DestHeight,
+	int32 NumTilesX, int32 NumTilesY,
+	int32 TileWidth, int32 TileHeight, int32 InTileBorder,
+	int32 BytesPerPixel)
+{
+	// We don't support tile borders larger than a tile size,
+	// but this shuld not happen in practice.
+	if ((InTileBorder > TileWidth) || (InTileBorder > TileHeight))
+	{
+		UE_LOG(LogImgMediaEditor, Error, TEXT("Tile border is larger than tile size. Clamping to tile size."));
+		InTileBorder = FMath::Min(TileWidth, TileHeight);
+	}
+
+	// Set up destination buffer.
+	DestArray.Reset();
+	DestArray.AddUninitialized(DestWidth * DestHeight * BytesPerPixel);
+
+	uint8* DestData = DestArray.GetData();
+	int32 DestTileWidth = TileWidth + InTileBorder * 2;
+	int32 DestTileHeight = TileHeight + InTileBorder * 2;
+
+	// Make sure our output tile size is not bigger than the output size.
+	if ((DestTileWidth > DestWidth) || (DestTileHeight > DestHeight))
+	{
+		// This is not a valid mip level, so just ignore.
+		return;
+	}
+
+	int32 BytesPerTile = TileWidth * TileHeight * BytesPerPixel;
+	int32 ByterPerDestTile = DestTileWidth * DestTileHeight * BytesPerPixel;
+
+	// Loop over y tiles.
+	for (int32 TileY = 0; TileY < NumTilesY; ++TileY)
+	{
+		// Loop over x tiles.
+		for (int32 TileX = 0; TileX < NumTilesX; ++TileX)
+		{
+			// Get address of the source and destination tiles.
+			uint8* SourceTile = SourceData +
+				(TileX * TileWidth + TileY * SourceWidth * TileHeight) * BytesPerPixel;
+			uint8* DestTile = DestData + (TileX + TileY * NumTilesX) * ByterPerDestTile;
+
+			int32 NumberOfPixelsToCopy = TileWidth;
+
+			// Create a left border.
+			if (TileX > 0)
+			{
+				NumberOfPixelsToCopy += InTileBorder;
+				// Offset the source to get the extra pixels.
+				SourceTile -= InTileBorder * BytesPerPixel;
+			}
+			else
+			{
+				// Offset the destination as we are skipping this border as we have no data.
+				DestTile += InTileBorder * BytesPerPixel;
+			}
+
+			// Create a right border.
+			if (TileX < NumTilesX - 1)
+			{
+				NumberOfPixelsToCopy += InTileBorder;
+			}
+
+			// Loop over each row in the tile.
+			for (int32 Row = 0; Row < DestTileHeight; ++Row)
+			{
+				// Make sure we don't go beyond the source data.
+				int32 SourceRow = Row - InTileBorder;
+				if (TileY == 0)
+				{
+					SourceRow = FMath::Max(SourceRow, 0);
+				}
+				if (TileY == NumTilesY - 1)
+				{
+					SourceRow = FMath::Min(SourceRow, TileHeight - 1);
+				}
+
+				uint8* SourceLine = SourceTile + SourceRow * SourceWidth * BytesPerPixel;
+				uint8* DestLine = DestTile + Row * DestTileWidth * BytesPerPixel;
+
+				// Copy the main data.
+				FMemory::Memcpy(DestLine, SourceLine, NumberOfPixelsToCopy * BytesPerPixel);
+			}
+		}
+	}
+}
+
 #undef LOCTEXT_NAMESPACE
diff --git a/Engine/Plugins/Media/ImgMedia/Source/ImgMediaEditor/Private/Widgets/SImgMediaProcessImages.h b/Engine/Plugins/Media/ImgMedia/Source/ImgMediaEditor/Private/Widgets/SImgMediaProcessImages.h
--- a/Engine/Plugins/Media/ImgMedia/Source/ImgMediaEditor/Private/Widgets/SImgMediaProcessImages.h
+++ b/Engine/Plugins/Media/ImgMedia/Source/ImgMediaEditor/Private/Widgets/SImgMediaProcessImages.h
@@ -76,6 +76,28 @@ private:
 		int32 InTileWidth, int32 InTileHeight, int32 InTileBorder, 
 		bool bInEnableMips, bool bHasAlphaChannel, const FString& InName);
 
+	/**
+	 * Creates tiles from a source and outputs it to a destination.
+	 *
+	 * @param SourceData		Source image.
+	 * @param DestArray			Destination image.
+	 * @param SourceWidth		Width of source in pixels.
+	 * @param SourceHeight		Height of source in pixels.
+	 * @param DestWidth			Width of destination in pixels.
+	 * @param DestHeight		Height of destination in pixels.
+	 * @param NumTilesX			Number of tiles in X direction.
+	 * @param NumTilesY			Number of tiles in Y direction.
+	 * @param TileWidth			Width of a tile (without borders) in pixels.
+	 * @param TileHeight		Height of a tile (without borders) in pixels.
+	 * @param InTileBorder		Size of border in pixels.
+	 * @param BytesPerPixel		Number of bytes per pixel.
+	 */
+	void TileData(uint8* SourceData, TArray64<uint8>& DestArray,
+		int32 SourceWidth, int32 SourceHeight, int32 DestWidth, int32 DestHeight,
+		int32 NumTilesX, int32 NumTilesY,
+		int32 TileWidth, int32 TileHeight, int32 InTileBorder,
+		int32 BytesPerPixel);
+
 	/** Holds our details view. */
 	TSharedPtr<class IDetailsView> DetailsView;
 	/** Object that holds our options. */
