ImgMedia ExrReaderGpu: Added ability to read and process custom Epic Exrs.

diff --git a/Engine/Plugins/Media/ImgMedia/Shaders/Private/ExrSwizzler.usf b/Engine/Plugins/Media/ImgMedia/Shaders/Private/ExrSwizzler.usf
--- a/Engine/Plugins/Media/ImgMedia/Shaders/Private/ExrSwizzler.usf
+++ b/Engine/Plugins/Media/ImgMedia/Shaders/Private/ExrSwizzler.usf
@@ -14,6 +14,9 @@ int2 TileSize;
 /** Total number of tiles in X and Y direction for the current mip level. */
 int2 NumTiles;
 
+/** Number of channels per pixel. */
+int NumChannels;
+
 /**
 * This function extracts the 16bit (float) data stored in uint (32 bit)
 */
@@ -32,9 +35,12 @@ float UintContainingFloatToFloat(in uint Value)
 /** Unpacks data from the provided structured buffer. */
 int GetValueFromTheBuffer(int Offset, int2 TextureSize, int ChannelIndex)
 {
+#if CUSTOM_EXR
+	int UnpackedLocation = Offset + ChannelIndex;
+#else
 	// the location as if it would've been if the buffer was unpacked.
 	int UnpackedLocation = Offset + TextureSize.x * ChannelIndex;
-
+#endif
 	// Byte offset within the 4 byte packed value.
 	int Remainder = UnpackedLocation % 2;
 
@@ -61,17 +67,7 @@ void MainVS(
 void GetStartScanlinePosition(in int2 PixelPos, out int2 SampleSize, out int StartPosition)
 {
 	// Depending on the EXR file it could contain multi-channel data.
-#if NUM_CHANNELS == 1
-	int Channels = 1;
-#elif NUM_CHANNELS == 2
-	int Channels = 2;
-#elif NUM_CHANNELS == 3
-	int Channels = 3;
-#else
-	int Channels = 4;
-#endif
-
-#if SWIZZLE_TILES
+#if RENDER_TILES
 	SampleSize = TileSize;
 
 	int xCoord = (PixelPos.x) % TileSize.x;
@@ -80,13 +76,20 @@ void GetStartScanlinePosition(in int2 PixelPos, out int2 SampleSize, out int Sta
 	int TileX = floor((PixelPos.x) / TileSize.x);
 	int TileY = floor((PixelPos.y) / TileSize.y);
 
-	// Padding in bytes
+	// Padding in buffer elements (uint)
 	const int TilePadding = 10;
 
 	// Current tile position.
-	int TileBufferStride = TileSize.x * TileSize.y * Channels;
+#if CUSTOM_EXR
+	int TileBufferStride = TileSize.x * TileSize.y * NumChannels;
+	int PreviousTilesOffset = (NumTiles.x * TileY + TileX) * TileBufferStride;
+	StartPosition = PreviousTilesOffset + yCoord * TileSize.x * NumChannels + xCoord * NumChannels;
+#else
+	int TileBufferStride = TileSize.x * TileSize.y * NumChannels;
 	int PreviousTilesOffset = (NumTiles.x * TileY + TileX) * TileBufferStride;
-	StartPosition = PreviousTilesOffset + yCoord * TileSize.x * Channels + xCoord + (TileY * NumTiles.x + TileX + 1) * TilePadding;
+	StartPosition = PreviousTilesOffset + yCoord * TileSize.x * NumChannels + xCoord + (TileY * NumTiles.x + TileX + 1) * TilePadding;
+#endif
+
 #else
 	SampleSize.x = TextureSize.x;
 	SampleSize.y = TextureSize.y;
@@ -95,8 +98,8 @@ void GetStartScanlinePosition(in int2 PixelPos, out int2 SampleSize, out int Sta
 	const int TilePadding = 4;
 
 	// Current scanline position.
-	StartPosition = PixelPos.y * TextureSize.x * Channels + PixelPos.x + (PixelPos.y + 1) * 4;
-#endif
+	StartPosition = PixelPos.y * TextureSize.x * NumChannels + PixelPos.x + (PixelPos.y + 1) * 4;
+#endif //RENDER_TILES
 }
 
 /** Pixel Shader. */
@@ -112,24 +115,33 @@ half4 MainPS(noperspective float4 UVAndScreenPos : TEXCOORD0) : SV_Target0
 	GetStartScanlinePosition(PixelPos, SampleSize, StartBufferPosition);
 
 	int ChannelIndex = 0;
-#if NUM_CHANNELS == 4
+
+#if NUM_CHANNELS == 3
 	half A = UintContainingFloatToFloat(GetValueFromTheBuffer(StartBufferPosition, SampleSize, ChannelIndex++));
 #else
 	half A = 1.;
 #endif
 
 	half B = UintContainingFloatToFloat(GetValueFromTheBuffer(StartBufferPosition, SampleSize, ChannelIndex++));
-#if NUM_CHANNELS <= 1
+#if NUM_CHANNELS == 0
 	half G = B;
 #else
 	half G = UintContainingFloatToFloat(GetValueFromTheBuffer(StartBufferPosition, SampleSize, ChannelIndex++));
 #endif
-#if NUM_CHANNELS <= 2
+#if NUM_CHANNELS <= 1
 	half R = B;
 #else
 	half R = UintContainingFloatToFloat(GetValueFromTheBuffer(StartBufferPosition, SampleSize, ChannelIndex++));
 #endif
 	half3 Rgb = half3(R, G, B);
 
+#if CUSTOM_EXR
+	#if NUM_CHANNELS == 3
+		return half4(A, Rgb.zyx);
+	#else
+		return half4(Rgb.zyx, A);
+	#endif
+#else
 	return half4(Rgb, A);
+#endif
 }
\ No newline at end of file
diff --git a/Engine/Plugins/Media/ImgMedia/Source/ExrReaderGpu/Private/ExrReaderGpu.cpp b/Engine/Plugins/Media/ImgMedia/Source/ExrReaderGpu/Private/ExrReaderGpu.cpp
--- a/Engine/Plugins/Media/ImgMedia/Source/ExrReaderGpu/Private/ExrReaderGpu.cpp
+++ b/Engine/Plugins/Media/ImgMedia/Source/ExrReaderGpu/Private/ExrReaderGpu.cpp
@@ -285,7 +285,7 @@ bool FExrReader::OpenExrAndPrepareForPixelReading(FString FilePath, int32 NumOff
 	return true;
 }
 
-bool FExrReader::ReadExrImageChunk(void* Buffer, int32 ChunkSize)
+bool FExrReader::ReadExrImageChunk(void* Buffer, int64 ChunkSize)
 {
 	if (FileHandle == nullptr)
 	{
@@ -306,7 +306,7 @@ bool FExrReader::ReadExrImageChunk(void* Buffer, int32 ChunkSize)
 	return NumElementsRead == NumElements;
 }
 
-bool FExrReader::SeekTileWithinFile(const int StartTileIndex, const FIntPoint& TexDimInTiles, int32& OutBufferOffset)
+bool FExrReader::SeekTileWithinFile(const int32 StartTileIndex, const FIntPoint& TexDimInTiles, int64& OutBufferOffset)
 {
 	if (StartTileIndex >= LineOrTileOffsets.Num())
 	{
@@ -318,6 +318,16 @@ bool FExrReader::SeekTileWithinFile(const int StartTileIndex, const FIntPoint& T
 	return fseek(FileHandle, LineOffset, SEEK_SET) == 0;
 }
 
+bool FExrReader::SeekTileWithinFileCustom(const int32 StartTileIndex, const int64 TileStride, int64& OutBufferOffset)
+{
+	// Our custom exr is written as one single tile. In Exr every tile starts with 20 byte padding.
+	// Therefore we want to ignore it, when we read custom tiles.
+	const int32 TilePadding = 20;
+	int64 LineOffset = TilePadding + LineOrTileOffsets[0] + StartTileIndex * TileStride;
+	OutBufferOffset = StartTileIndex * TileStride;
+	return fseek(FileHandle, LineOffset, SEEK_SET) == 0;
+}
+
 bool FExrReader::CloseExrFile()
 {
 	if (FileHandle == nullptr)
diff --git a/Engine/Plugins/Media/ImgMedia/Source/ExrReaderGpu/Public/ExrReaderGpu.h b/Engine/Plugins/Media/ImgMedia/Source/ExrReaderGpu/Public/ExrReaderGpu.h
--- a/Engine/Plugins/Media/ImgMedia/Source/ExrReaderGpu/Public/ExrReaderGpu.h
+++ b/Engine/Plugins/Media/ImgMedia/Source/ExrReaderGpu/Public/ExrReaderGpu.h
@@ -31,7 +31,13 @@ public:
 private:
 	static const int32	STRING_SIZE = 256;
 	static const int32	MAX_LENGTH = STRING_SIZE - 1;
+
+public:
+	/*At the beginning of each row of B G R channel planes there is 2x4 byte data that has information*/
 	static const int32	PLANAR_RGB_SCANLINE_PADDING = 8;
+	/*At the beginning of each tile there is 20 byte data that has information*/
+	static const int32	TILE_PADDING = 20;
+
 public:
 	FExrReader() :FileHandle(nullptr) {};
 private:
@@ -64,12 +70,17 @@ public:
 	* Read a chunk of an Exr file previously open via OpenExrAndPrepareForPixelReading.
 	* The read starts from the last point of file reading.
 	*/
-	bool ReadExrImageChunk(void* Buffer, int32 ChunkSize);
+	bool ReadExrImageChunk(void* Buffer, int64 ChunkSize);
+
+	/** 
+	* Moves to the specified tile position in the file in preparation for reading.
+	*/
+	bool SeekTileWithinFile(const int32 StartTileIndex, const FIntPoint& TexDimInTiles, int64& OutBufferOffset);
 
 	/** 
 	* Moves to the specified tile position in the file in preparation for reading.
 	*/
-	bool SeekTileWithinFile(const int StartTileIndex, const FIntPoint& TexDimInTiles, int32& OutBufferOffset);
+	bool SeekTileWithinFileCustom(const int32 StartTileIndex, const int64 TileStride, int64& OutBufferOffset);
 
 	/**
 	* After we have finished reading the file that was open via OpenExrAndPrepareForPixelReading
@@ -78,6 +89,7 @@ public:
 	bool CloseExrFile();
 
 private:
+
 	/**
 	* One approach to reading is to first open a file, read in chunks and then close the file.
 	* It is all done in 3 sepate functions and requires to keep track of Filehandle
diff --git a/Engine/Plugins/Media/ImgMedia/Source/ExrReaderGpu/Public/ExrSwizzlingShader.h b/Engine/Plugins/Media/ImgMedia/Source/ExrReaderGpu/Public/ExrSwizzlingShader.h
--- a/Engine/Plugins/Media/ImgMedia/Source/ExrReaderGpu/Public/ExrSwizzlingShader.h
+++ b/Engine/Plugins/Media/ImgMedia/Source/ExrReaderGpu/Public/ExrSwizzlingShader.h
@@ -40,9 +40,10 @@ class EXRREADERGPU_API FExrSwizzlePS : public FGlobalShader
 	SHADER_USE_PARAMETER_STRUCT(FExrSwizzlePS, FGlobalShader);
 
 	/** If the provided buffer is RGBA the shader would work slightly differently to RGB. */
-	class FRgbaSwizzle : SHADER_PERMUTATION_INT("NUM_CHANNELS", 5);
-	class FSwizzleTiles : SHADER_PERMUTATION_BOOL("SWIZZLE_TILES");
-	using FPermutationDomain = TShaderPermutationDomain<FRgbaSwizzle, FSwizzleTiles>;
+	class FRgbaSwizzle : SHADER_PERMUTATION_INT("NUM_CHANNELS", 4);
+	class FRenderTiles : SHADER_PERMUTATION_BOOL("RENDER_TILES");
+	class FCustomExr : SHADER_PERMUTATION_BOOL("CUSTOM_EXR");
+	using FPermutationDomain = TShaderPermutationDomain<FRgbaSwizzle, FRenderTiles, FCustomExr>;
 
 	static bool ShouldCompilePermutation(const FGlobalShaderPermutationParameters& Parameters) 
 	{
@@ -54,6 +55,7 @@ class EXRREADERGPU_API FExrSwizzlePS : public FGlobalShader
 		SHADER_PARAMETER(FIntPoint, TextureSize)
 		SHADER_PARAMETER(FIntPoint, TileSize)
 		SHADER_PARAMETER(FIntPoint, NumTiles)
+		SHADER_PARAMETER(int32, NumChannels)
 	END_SHADER_PARAMETER_STRUCT()
 };
 #endif
\ No newline at end of file
diff --git a/Engine/Plugins/Media/ImgMedia/Source/ImgMedia/Private/Readers/ExrImgMediaReader.cpp b/Engine/Plugins/Media/ImgMedia/Source/ImgMedia/Private/Readers/ExrImgMediaReader.cpp
--- a/Engine/Plugins/Media/ImgMedia/Source/ImgMedia/Private/Readers/ExrImgMediaReader.cpp
+++ b/Engine/Plugins/Media/ImgMedia/Source/ImgMedia/Private/Readers/ExrImgMediaReader.cpp
@@ -7,6 +7,7 @@
 
 #include "Async/Async.h"
 #include "Misc/Paths.h"
+#include "ExrReaderGpu.h"
 #include "OpenExrWrapper.h"
 #include "UObject/Class.h"
 #include "UObject/UObjectGlobals.h"
@@ -70,6 +71,68 @@ bool FExrImgMediaReader::GetFrameInfo(const FString& ImagePath, FImgMediaFrameIn
 	return false;
 }
 
+FExrImgMediaReader::EReadResult FExrImgMediaReader::ReadTilesCustom
+(uint16* Buffer
+	, const FString& ImagePath
+	, int32 FrameId
+	, const FIntRect& TileRegion
+	, const FIntPoint& FullTexDimInTiles
+	, const FIntPoint& TileDim
+	, const int32 PixelSize
+	, const bool bCustomExr)
+{
+	EReadResult bResult = Success;
+
+	int64 CurrentBufferPos = 0;
+	FExrReader ChunkReader;
+	int32 NumTiles = bCustomExr ? 1 : FullTexDimInTiles.X * FullTexDimInTiles.Y;
+	if (!ChunkReader.OpenExrAndPrepareForPixelReading(ImagePath, NumTiles))
+	{
+		return Fail;
+	}
+
+	for (int32 TileRow = TileRegion.Min.Y; TileRow < TileRegion.Max.Y; TileRow++)
+	{
+		// Check to see if the frame was canceled.
+		{
+			FScopeLock RegionScopeLock(&CanceledFramesCriticalSection);
+			if (CanceledFrames.Remove(FrameId) > 0)
+			{
+				UE_LOG(LogImgMedia, Warning, TEXT("Reader %p: Canceling Frame %i At tile row # %i"), this, FrameId, TileRow);
+				bResult = Cancelled;
+				break;
+			}
+		}
+		const uint16 Padding = bCustomExr ? 0 : FExrReader::TILE_PADDING;
+		const int64 TileByteStride = PixelSize * TileDim.X * TileDim.Y + Padding;
+		const int StartTileIndex = TileRow * FullTexDimInTiles.X + TileRegion.Min.X;
+		const int EndTileIndex = TileRow * FullTexDimInTiles.X + TileRegion.Max.X;
+		const int BytesToRead = (EndTileIndex - StartTileIndex) * TileByteStride;
+
+		if (bCustomExr)
+		{
+			int64 TileStride = TileDim.X * TileDim.Y * PixelSize;
+			ChunkReader.SeekTileWithinFileCustom(StartTileIndex, TileStride, CurrentBufferPos);
+		}
+		else
+		{
+			ChunkReader.SeekTileWithinFile(StartTileIndex, FullTexDimInTiles, CurrentBufferPos);
+		}
+
+		if (!ChunkReader.ReadExrImageChunk(reinterpret_cast<char*>(Buffer) + CurrentBufferPos, BytesToRead))
+		{
+			bResult = Fail;
+			break;
+		}
+	}
+
+	if (!ChunkReader.CloseExrFile())
+	{
+		return Fail;
+	}
+
+	return bResult;
+}
 
 bool FExrImgMediaReader::ReadFrame(int32 FrameId, const TMap<int32, FImgMediaTileSelection>& InMipTiles, TSharedPtr<FImgMediaFrame, ESPMode::ThreadSafe> OutFrame)
 {
@@ -183,40 +246,63 @@ bool FExrImgMediaReader::ReadFrame(int32 FrameId, const TMap<int32, FImgMediaTil
 				int32 FrameBufferOffsetPerTileY = Dim.X * TileHeight * BytesPerPixel;
 				int32 FrameBufferOffsetY = FrameBufferOffsetPerTileY * StartTileY;
 
-				for (int32 TileY = StartTileY; TileY < EndTileY; TileY++)
+				if (OutFrame->Info.FormatName == TEXT("EXR CUSTOM"))
 				{
-					// The offset into the frame buffer for each column of tiles.
-					// Tile width * bytes per pixel.
-					int32 FrameBufferOffsetPerTileX = TileWidth * BytesPerPixel;
-					int32 FrameBufferOffsetX = FrameBufferOffsetPerTileX * StartTileX;
-					for (int32 TileX = StartTileX; TileX < EndTileX; TileX++)
+					FIntRect TileRegion = FIntRect(
+						(int32)CurrentTileSelection.TopLeftX,
+						(int32)CurrentTileSelection.TopLeftY,
+						FMath::Min((int32)CurrentTileSelection.BottomRightX, FMath::CeilToInt(float(OutFrame->Info.NumTiles.X) / MipLevelDiv)),
+						FMath::Min((int32)CurrentTileSelection.BottomRightY, FMath::CeilToInt(float(OutFrame->Info.NumTiles.Y) / MipLevelDiv)));
+					FIntPoint NumTiles = OutFrame->Info.NumTiles / MipLevelDiv;
+					int32 PixelSize = sizeof(uint16) * OutFrame->Info.NumChannels;
+
+					EReadResult ReadResult = ReadTilesCustom((uint16*)MipDataPtr, Image, FrameId, TileRegion, NumTiles, OutFrame->Info.TileDimensions, PixelSize, true /*bCustomExr*/);
+					if (ReadResult != Fail)
 					{
-						// Get for our frame/mip level.
-						if (bHasTiles)
-						{
-							Image = FString::Printf(TEXT("%s_x%d_y%d.exr"), *BaseImage, TileX, TileY);
-						}
-						FRgbaInputFile InputFile(Image, 2);
-						if (InputFile.HasInputFile())
+						OutFrame->MipTilesPresent.Emplace(CurrentMipLevel, CurrentTileSelection);
+					}
+					else
+					{
+						UE_LOG(LogImgMedia, Error, TEXT("Could not load %s"), *Image);
+					}
+				}
+				else
+				{
+					for (int32 TileY = StartTileY; TileY < EndTileY; TileY++)
+					{
+						// The offset into the frame buffer for each column of tiles.
+						// Tile width * bytes per pixel.
+						int32 FrameBufferOffsetPerTileX = TileWidth * BytesPerPixel;
+						int32 FrameBufferOffsetX = FrameBufferOffsetPerTileX * StartTileX;
+						for (int32 TileX = StartTileX; TileX < EndTileX; TileX++)
 						{
-							// read frame data
-							InputFile.SetFrameBuffer(MipDataPtr + FrameBufferOffsetX + FrameBufferOffsetY, Dim);
-							InputFile.ReadPixels(0, TileHeight - 1);
+							// Get for our frame/mip level.
+							if (bHasTiles)
+							{
+								Image = FString::Printf(TEXT("%s_x%d_y%d.exr"), *BaseImage, TileX, TileY);
+							}
+							FRgbaInputFile InputFile(Image, 2);
+							if (InputFile.HasInputFile())
+							{
+								// read frame data
+								InputFile.SetFrameBuffer(MipDataPtr + FrameBufferOffsetX + FrameBufferOffsetY, Dim);
+								InputFile.ReadPixels(0, TileHeight - 1);
 
-							FrameBufferOffsetX += FrameBufferOffsetPerTileX;
+								FrameBufferOffsetX += FrameBufferOffsetPerTileX;
 
-							if (!OutFrame->MipTilesPresent.Contains(CurrentMipLevel))
+								if (!OutFrame->MipTilesPresent.Contains(CurrentMipLevel))
+								{
+									OutFrame->MipTilesPresent.Emplace(CurrentMipLevel, CurrentTileSelection);
+								}
+							}
+							else
 							{
-								OutFrame->MipTilesPresent.Emplace(CurrentMipLevel, CurrentTileSelection);
+								UE_LOG(LogImgMedia, Error, TEXT("Could not load %s"), *Image);
 							}
 						}
-						else
-						{
-							UE_LOG(LogImgMedia, Error, TEXT("Could not load %s"), *Image);
-						}
-					}
 
-					FrameBufferOffsetY += FrameBufferOffsetPerTileY;
+						FrameBufferOffsetY += FrameBufferOffsetPerTileY;
+					}
 				}
 			}
 		}
@@ -297,7 +383,6 @@ bool FExrImgMediaReader::GetInfo(const FString& FilePath, FImgMediaFrameInfo& Ou
 
 	OutInfo.CompressionName = HeaderReader.GetCompressionName();
 	OutInfo.Dim = HeaderReader.GetDataWindow();
-	OutInfo.FormatName = TEXT("EXR");
 	OutInfo.FrameRate = HeaderReader.GetFrameRate(ImgMedia::DefaultFrameRate);
 	OutInfo.Srgb = false;
 	OutInfo.UncompressedSize = HeaderReader.GetUncompressedSize();
@@ -306,8 +391,11 @@ bool FExrImgMediaReader::GetInfo(const FString& FilePath, FImgMediaFrameInfo& Ou
 	int32 CustomFormat = 0;
 	HeaderReader.GetIntAttribute(IImgMediaModule::CustomFormatAttributeName.Resolve().ToString(), CustomFormat);
 	bool bIsCustomFormat = CustomFormat > 0;
+
 	if (bIsCustomFormat)
 	{
+		OutInfo.FormatName = TEXT("EXR CUSTOM");
+
 		// Get tile size.
 		HeaderReader.GetIntAttribute(IImgMediaModule::CustomFormatTileBorderAttributeName.Resolve().ToString(), OutInfo.TileBorder);
 		OutInfo.bHasTiles = HeaderReader.GetIntAttribute(IImgMediaModule::CustomFormatTileWidthAttributeName.Resolve().ToString(), OutInfo.TileDimensions.X);
@@ -315,6 +403,7 @@ bool FExrImgMediaReader::GetInfo(const FString& FilePath, FImgMediaFrameInfo& Ou
 	}
 	else
 	{
+		OutInfo.FormatName = TEXT("EXR");
 		OutInfo.bHasTiles = HeaderReader.GetTileSize(OutInfo.TileDimensions);
 		OutInfo.TileBorder = 0;
 	}
diff --git a/Engine/Plugins/Media/ImgMedia/Source/ImgMedia/Private/Readers/ExrImgMediaReader.h b/Engine/Plugins/Media/ImgMedia/Source/ImgMedia/Private/Readers/ExrImgMediaReader.h
--- a/Engine/Plugins/Media/ImgMedia/Source/ImgMedia/Private/Readers/ExrImgMediaReader.h
+++ b/Engine/Plugins/Media/ImgMedia/Source/ImgMedia/Private/Readers/ExrImgMediaReader.h
@@ -42,16 +42,28 @@ public:
 	FIntPoint GetCustomFormatTileSize() { return CustomFormatTileSize; }
 
 protected:
+	enum EReadResult
+	{
+		Fail,
+		Success,
+		Cancelled
+	};
 
 	/**
 	 * Get the frame information from the given input file.
 	 *
-	 * @param InputFile The input file containing the frame.
+	 * @param FilePath The location of the exr file.
 	 * @param OutInfo Will contain the frame information.
 	 * @return true on success, false otherwise.
 	 */
 	static bool GetInfo(const FString& FilePath, FImgMediaFrameInfo& OutInfo);
 
+	/**
+	 * Reads tiles from exr files in tile rows based on Tile region. If frame is pending for cancelation
+	 * stops reading tiles at the current tile row.
+	*/
+	EReadResult ReadTilesCustom(uint16* Buffer, const FString& ImagePath, int32 FrameId, const FIntRect& TileRegion, const FIntPoint& FullTexDimInTiles, const FIntPoint& TileDim, const int32 PixelSize, const bool bCustomExr);
+
 	/**
 	 * Sets parameters of our custom format images.
 	 *
diff --git a/Engine/Plugins/Media/ImgMedia/Source/ImgMedia/Private/Readers/ExrImgMediaReaderGpu.cpp b/Engine/Plugins/Media/ImgMedia/Source/ImgMedia/Private/Readers/ExrImgMediaReaderGpu.cpp
--- a/Engine/Plugins/Media/ImgMedia/Source/ImgMedia/Private/Readers/ExrImgMediaReaderGpu.cpp
+++ b/Engine/Plugins/Media/ImgMedia/Source/ImgMedia/Private/Readers/ExrImgMediaReaderGpu.cpp
@@ -33,12 +33,6 @@ DECLARE_GPU_STAT_NAMED(ExrImgMediaReaderGpu, TEXT("ExrImgMediaReaderGpu"));
 
 namespace {
 
-	/* Padding in bits that is added at the beginning of each scanline */
-	const uint16 CExrRowPadding = 8;
-
-	/* Padding in bits that is added at the beginning of each tile */
-	const uint16 CExrTilePadding = 20;
-
 	/** This function is similar to DrawScreenPass in OpenColorIODisplayExtension.cpp except it is catered for Viewless texture rendering. */
 	template<typename TSetupFunction>
 	void DrawScreenPass(
@@ -145,27 +139,28 @@ bool FExrImgMediaReaderGpu::ReadFrame(int32 FrameId, const TMap<int32, FImgMedia
 	// Get tile info.
 	bool bHasTiles = OutFrame->Info.bHasTiles;
 
-	const FIntPoint& FullResolution = OutFrame->Info.Dim;
+	TSharedPtr<FSampleConverterParameters> ConverterParams = MakeShared<FSampleConverterParameters>();
+	ConverterParams->FullResolution = OutFrame->Info.Dim;
 
 	const FIntPoint TileDim = OutFrame->Info.TileDimensions;
 
-	if (FullResolution.GetMin() <= 0)
+	if (ConverterParams->FullResolution.GetMin() <= 0)
 	{
 		return false;
 	}
 
-	TMap<int32, FIntRect> Viewports;
-
-	const int32 NumChannels = OutFrame->Info.NumChannels;
-	const int32 PixelSize = sizeof(uint16) * NumChannels;
-
 	TArray<FStructuredBufferPoolItemSharedPtr> BufferDataArray;
 	BufferDataArray.AddDefaulted(FImgMediaLoader::MAX_MIPMAP_LEVELS);
 
-	int32 NumMipLevels = Loader->GetNumMipLevels();
+	ConverterParams->FrameInfo = OutFrame->Info;
+	ConverterParams->PixelSize = sizeof(uint16) * ConverterParams->FrameInfo.NumChannels;
+	ConverterParams->TileDimWithBorders = TileDim + OutFrame->Info.TileBorder * 2;
+	ConverterParams->NumMipLevels = Loader->GetNumMipLevels();
+	ConverterParams->bCustomExr = OutFrame->Info.FormatName == TEXT("EXR CUSTOM");
+
 	{
 		// Loop over all mips.
-		FIntPoint CurrentMipDim = FullResolution;
+		FIntPoint CurrentMipDim = ConverterParams->FullResolution;
 
 		for (const TPair<int32, FImgMediaTileSelection>& TilesPerMip : InMipTiles)
 		{
@@ -181,16 +176,16 @@ bool FExrImgMediaReaderGpu::ReadFrame(int32 FrameId, const TMap<int32, FImgMedia
 			
 			// Next mip level.
 			int MipLevelDiv = 1 << CurrentMipLevel;
-			CurrentMipDim = FullResolution / MipLevelDiv;
+			CurrentMipDim = ConverterParams->FullResolution / MipLevelDiv;
 
 			if (ReadThisMip)
 			{
-				FIntRect& Viewport = Viewports.Add(CurrentMipLevel);
-				Viewport.Min = FIntPoint(TileDim.X * CurrentTileSelection.TopLeftX, TileDim.X * CurrentTileSelection.TopLeftY);
-				Viewport.Max = FIntPoint(TileDim.Y * CurrentTileSelection.BottomRightX, TileDim.Y * CurrentTileSelection.BottomRightY);
+				FIntRect& Viewport = ConverterParams->Viewports.Add(CurrentMipLevel);
+				Viewport.Min = FIntPoint(ConverterParams->TileDimWithBorders.X * CurrentTileSelection.TopLeftX, ConverterParams->TileDimWithBorders.X * CurrentTileSelection.TopLeftY);
+				Viewport.Max = FIntPoint(ConverterParams->TileDimWithBorders.Y * CurrentTileSelection.BottomRightX, ConverterParams->TileDimWithBorders.Y * CurrentTileSelection.BottomRightY);
 				Viewport.Clip(FIntRect(FIntPoint::ZeroValue, CurrentMipDim));
 
-				const SIZE_T BufferSize = GetBufferSize(CurrentMipDim, NumChannels, bHasTiles, OutFrame->Info.NumTiles / MipLevelDiv);
+				const SIZE_T BufferSize = GetBufferSize(CurrentMipDim, ConverterParams->FrameInfo.NumChannels, bHasTiles, OutFrame->Info.NumTiles / MipLevelDiv, ConverterParams->bCustomExr);
 				FStructuredBufferPoolItemSharedPtr& BufferData = BufferDataArray[CurrentMipLevel];
 				BufferData = AllocateGpuBufferFromPool(BufferSize);
 				uint16* MipDataPtr = static_cast<uint16*>(BufferData->MappedBuffer);
@@ -204,7 +199,7 @@ bool FExrImgMediaReaderGpu::ReadFrame(int32 FrameId, const TMap<int32, FImgMedia
 				if (InputTileFile.HasInputFile())
 				{
 					// read frame data
-					if (bHasTiles)
+					if (bHasTiles || ConverterParams->bCustomExr)
 					{
 						FIntRect TileRegion = FIntRect(
 							(int32)CurrentTileSelection.TopLeftX,
@@ -213,7 +208,7 @@ bool FExrImgMediaReaderGpu::ReadFrame(int32 FrameId, const TMap<int32, FImgMedia
 							FMath::Min((int32)CurrentTileSelection.BottomRightY, FMath::CeilToInt(float(OutFrame->Info.NumTiles.Y) / MipLevelDiv)));
 						FIntPoint NumTiles = OutFrame->Info.NumTiles / MipLevelDiv;
 
-						ReadResult = ReadTiles(MipDataPtr, ImagePath, FrameId, TileRegion, NumTiles, TileDim, PixelSize);
+						ReadResult = ReadTilesCustom(MipDataPtr, ImagePath, FrameId, TileRegion, NumTiles, ConverterParams->TileDimWithBorders, ConverterParams->PixelSize, ConverterParams->bCustomExr);
 					}
 					else
 					{
@@ -233,10 +228,14 @@ bool FExrImgMediaReaderGpu::ReadFrame(int32 FrameId, const TMap<int32, FImgMedia
 				if (ReadResult == Fail)
 				{
 					// Check if we have a compressed file.
-					if (OutFrame->Info.CompressionName != "Uncompressed")
+					FImgMediaFrameInfo Info;
+					if (GetInfo(ImagePath, Info))
 					{
-						UE_LOG(LogImgMedia, Error, TEXT("GPU Reader cannot read compressed file %s."), *ImagePath);
-						UE_LOG(LogImgMedia, Error, TEXT("Compressed and uncompressed files should not be mixed in a single sequence."));
+						if (Info.CompressionName != "Uncompressed")
+						{
+							UE_LOG(LogImgMedia, Error, TEXT("GPU Reader cannot read compressed file %s."), *ImagePath);
+							UE_LOG(LogImgMedia, Error, TEXT("Compressed and uncompressed files should not be mixed in a single sequence."));
+						}
 					}
 
 					// Fall back to CPU.
@@ -246,18 +245,17 @@ bool FExrImgMediaReaderGpu::ReadFrame(int32 FrameId, const TMap<int32, FImgMedia
 			}
 		}
 	}
-	OutFrame->Format = NumChannels <= 3 ? EMediaTextureSampleFormat::FloatRGB : EMediaTextureSampleFormat::FloatRGBA;
-	OutFrame->Stride = FullResolution.X * PixelSize;
-
-	OutFrame->SampleConverter = CreateSampleConverter(MoveTemp(BufferDataArray), FullResolution, TileDim, Viewports, NumChannels, NumMipLevels, bHasTiles);
+	OutFrame->Format = ConverterParams->FrameInfo.NumChannels <= 3 ? EMediaTextureSampleFormat::FloatRGB : EMediaTextureSampleFormat::FloatRGBA;
+	OutFrame->Stride = ConverterParams->FullResolution.X * ConverterParams->PixelSize;
+	OutFrame->SampleConverter = CreateSampleConverter(MoveTemp(BufferDataArray), ConverterParams);
 
 	UE_LOG(LogImgMedia, Verbose, TEXT("Reader %p: Read Pixels Complete. %i"), this, FrameId);
 	return true;
 }
 
-void FExrImgMediaReaderGpu::PreAllocateMemoryPool(int32 NumFrames, const FImgMediaFrameInfo& FrameInfo)
+void FExrImgMediaReaderGpu::PreAllocateMemoryPool(int32 NumFrames, const FImgMediaFrameInfo& FrameInfo, const bool bCustomExr)
 {
-	SIZE_T AllocSize = GetBufferSize(FrameInfo.Dim, FrameInfo.NumChannels, FrameInfo.bHasTiles, FrameInfo.NumTiles);
+	SIZE_T AllocSize = GetBufferSize(FrameInfo.Dim, FrameInfo.NumChannels, FrameInfo.bHasTiles, FrameInfo.NumTiles, bCustomExr);
 	for (int32 FrameCacheNum = 0; FrameCacheNum < NumFrames; FrameCacheNum++)
 	{
 		AllocateGpuBufferFromPool(AllocSize, FrameCacheNum == NumFrames - 1);
@@ -331,58 +329,9 @@ FExrImgMediaReaderGpu::EReadResult FExrImgMediaReaderGpu::ReadInChunks(uint16* B
 	return bResult;
 }
 
-FExrImgMediaReaderGpu::EReadResult FExrImgMediaReaderGpu::ReadTiles(uint16* Buffer, const FString& ImagePath, int32 FrameId, const FIntRect& TileRegion, const FIntPoint& FullTexDimInTiles, const FIntPoint& TileDim, const int32 PixelSize)
-{
-	EReadResult bResult = Success;
-
-	// Chunks are of 16 MB
-	int32 CurrentBufferPos = 0;
-	FExrReader ChunkReader;
-
-
-	if (!ChunkReader.OpenExrAndPrepareForPixelReading(ImagePath, FullTexDimInTiles.X * FullTexDimInTiles.Y))
-	{
-		return Fail;
-	}
-
-	for (int32 TileRow = TileRegion.Min.Y; TileRow < TileRegion.Max.Y; TileRow++)
-	{
-		// Check to see if the frame was canceled.
-		{
-			FScopeLock RegionScopeLock(&CanceledFramesCriticalSection);
-			if (CanceledFrames.Remove(FrameId) > 0)
-			{
-				UE_LOG(LogImgMedia, Warning, TEXT("Reader %p: Canceling Frame %i At tile row # %i"), this, FrameId, TileRow);
-				bResult = Cancelled;
-				break;
-			}
-		}
-
-		const int32 TileByteStride = PixelSize * TileDim.X * TileDim.Y + CExrTilePadding;
-		const int StartTileIndex = TileRow * FullTexDimInTiles.X + TileRegion.Min.X;
-		const int EndTileIndex = TileRow * FullTexDimInTiles.X + TileRegion.Max.X;
-		const int BytesToRead = (EndTileIndex - StartTileIndex) * TileByteStride;
-
-		ChunkReader.SeekTileWithinFile(StartTileIndex, FullTexDimInTiles, CurrentBufferPos);
-
-		if (!ChunkReader.ReadExrImageChunk(reinterpret_cast<char*>(Buffer) + CurrentBufferPos, BytesToRead))
-		{
-			bResult = Fail;
-			break;
-		}
-	}
-
-	if (!ChunkReader.CloseExrFile())
-	{
-		return Fail;
-	}
-
-	return bResult;
-}
-
-SIZE_T FExrImgMediaReaderGpu::GetBufferSize(const FIntPoint& Dim, int32 NumChannels, bool bHasTiles, const FIntPoint& TileNum)
+SIZE_T FExrImgMediaReaderGpu::GetBufferSize(const FIntPoint& Dim, int32 NumChannels, bool bHasTiles, const FIntPoint& TileNum, const bool bCustomExr)
 {
-	if (!bHasTiles)
+	if (!bHasTiles && !bCustomExr)
 	{
 		/** 
 		* Reading scanlines.
@@ -390,7 +339,8 @@ SIZE_T FExrImgMediaReaderGpu::GetBufferSize(const FIntPoint& Dim, int32 NumChann
 		* At the beginning of each row of B G R channel planes there is 2x4 byte data that has information
 		* about number of pixels in the current row and row's number.
 		*/
-		SIZE_T BufferSize = Dim.X * Dim.Y * sizeof(uint16) * NumChannels + Dim.Y * CExrRowPadding;
+		const uint16 Padding = FExrReader::PLANAR_RGB_SCANLINE_PADDING;
+		SIZE_T BufferSize = Dim.X * Dim.Y * sizeof(uint16) * NumChannels + Dim.Y * Padding;
 		return BufferSize;
 	}
 	else
@@ -398,10 +348,11 @@ SIZE_T FExrImgMediaReaderGpu::GetBufferSize(const FIntPoint& Dim, int32 NumChann
 		/** 
 		* Reading tiles.
 		* 
-		* At the beginning of each tile of B G R channel planes there is 10 byte data that has information
+		* At the beginning of each tile there is 20 byte data that has information
 		* about number contents of tiles.
 		*/
-		SIZE_T BufferSize = Dim.X * Dim.Y * sizeof(uint16) * NumChannels + TileNum.X * TileNum.Y * CExrTilePadding;
+		const uint16 Padding = bCustomExr ? 0 : FExrReader::TILE_PADDING;
+		SIZE_T BufferSize = Dim.X * Dim.Y * sizeof(uint16) * NumChannels + (TileNum.X * TileNum.Y) * Padding;
 		return BufferSize;
 	}
 
@@ -409,24 +360,19 @@ SIZE_T FExrImgMediaReaderGpu::GetBufferSize(const FIntPoint& Dim, int32 NumChann
 
 TSharedPtr<IMediaTextureSampleConverter, ESPMode::ThreadSafe> FExrImgMediaReaderGpu::CreateSampleConverter
 	( TArray<FStructuredBufferPoolItemSharedPtr>&& BufferDataArray
-	, const FIntPoint& FullResolution
-	, const FIntPoint& TileDim
-	, const TMap<int32, FIntRect>& Viewports
-	, const int32 NumChannels
-	, const int32 NumMipLevels
-	, const bool bHasTiles)
+	, TSharedPtr<FSampleConverterParameters> ConverterParams)
 {
-	auto RenderThreadSwizzler = [this, BufferDataArray, FullResolution, TileDim, Viewports, NumChannels, NumMipLevels, bHasTiles](FRHICommandListImmediate& RHICmdList, FTexture2DRHIRef RenderTargetTextureRHI)->bool
+	auto RenderThreadSwizzler = [BufferDataArray, ConverterParams] (FRHICommandListImmediate& RHICmdList, FTexture2DRHIRef RenderTargetTextureRHI)->bool
 	{
 
 		SCOPED_DRAW_EVENT(RHICmdList, FExrImgMediaReaderGpu_Convert);
 		SCOPED_GPU_STAT(RHICmdList, ExrImgMediaReaderGpu);
 
-		FIntPoint Dim = FullResolution;
-		for (int32 MipLevel = 0; MipLevel < NumMipLevels; ++MipLevel)
+		FIntPoint Dim = ConverterParams->FullResolution;
+		for (int32 MipLevel = 0; MipLevel < ConverterParams->NumMipLevels; ++MipLevel)
 		{
 			int MipLevelDiv = 1 << MipLevel;
-			Dim = FullResolution / MipLevelDiv;
+			Dim = ConverterParams->FullResolution / MipLevelDiv;
 
 			FStructuredBufferPoolItemSharedPtr BufferData = BufferDataArray[MipLevel];
 			if (BufferData.IsValid())
@@ -436,7 +382,7 @@ TSharedPtr<IMediaTextureSampleConverter, ESPMode::ThreadSafe> FExrImgMediaReader
 					continue;
 				}
 
-				FIntRect Viewport = Viewports[MipLevel];
+				FIntRect Viewport = ConverterParams->Viewports[MipLevel];
 
 				// This flag will indicate that we should wait for poll to complete.
 				BufferData->bWillBeSignaled = true;
@@ -445,14 +391,16 @@ TSharedPtr<IMediaTextureSampleConverter, ESPMode::ThreadSafe> FExrImgMediaReader
 				RHICmdList.BeginRenderPass(RPInfo, TEXT("ExrTextureSwizzle"));
 
 				FExrSwizzlePS::FPermutationDomain PermutationVector;
-				PermutationVector.Set<FExrSwizzlePS::FRgbaSwizzle>(NumChannels);
-				PermutationVector.Set<FExrSwizzlePS::FSwizzleTiles>(bHasTiles);
+				PermutationVector.Set<FExrSwizzlePS::FRgbaSwizzle>(ConverterParams->FrameInfo.NumChannels - 1);
+				PermutationVector.Set<FExrSwizzlePS::FRenderTiles>(ConverterParams->FrameInfo.bHasTiles || ConverterParams->bCustomExr);
+				PermutationVector.Set<FExrSwizzlePS::FCustomExr>(ConverterParams->bCustomExr);
 				FExrSwizzlePS::FParameters Parameters = FExrSwizzlePS::FParameters();
 				Parameters.TextureSize = Dim;
-				Parameters.TileSize = TileDim;
-				if (bHasTiles)
+				Parameters.TileSize = ConverterParams->TileDimWithBorders;
+				Parameters.NumChannels = ConverterParams->FrameInfo.NumChannels;
+				if (ConverterParams->FrameInfo.bHasTiles)
 				{
-					Parameters.NumTiles = Dim / TileDim;
+					Parameters.NumTiles = Dim / ConverterParams->TileDimWithBorders;
 				}
 
 				Parameters.UnswizzledBuffer = RHICreateShaderResourceView(BufferData->BufferRef);
diff --git a/Engine/Plugins/Media/ImgMedia/Source/ImgMedia/Private/Readers/ExrImgMediaReaderGpu.h b/Engine/Plugins/Media/ImgMedia/Source/ImgMedia/Private/Readers/ExrImgMediaReaderGpu.h
--- a/Engine/Plugins/Media/ImgMedia/Source/ImgMedia/Private/Readers/ExrImgMediaReaderGpu.h
+++ b/Engine/Plugins/Media/ImgMedia/Source/ImgMedia/Private/Readers/ExrImgMediaReaderGpu.h
@@ -67,16 +67,23 @@ public:
 	/**
 	* For performance reasons we want to pre-allocate structured buffers to at least the number of concurrent frames.
 	*/
-	virtual void PreAllocateMemoryPool(int32 NumFrames, const FImgMediaFrameInfo& FrameInfo) override;
+	virtual void PreAllocateMemoryPool(int32 NumFrames, const FImgMediaFrameInfo& FrameInfo, const bool bCustomExr) override;
 	virtual void OnTick() override;
 
 protected:
 
-	enum EReadResult
+	/* 
+	* These are all the required parameters to produce Buffer to Texture converter. 
+	*/
+	struct FSampleConverterParameters
 	{
-		Fail,
-		Success,
-		Cancelled
+		FImgMediaFrameInfo FrameInfo;
+		FIntPoint FullResolution;
+		FIntPoint TileDimWithBorders;
+		TMap<int32, FIntRect> Viewports;
+		int32 NumMipLevels;
+		int32 PixelSize;
+		bool bCustomExr;
 	};
 
 	/** 
@@ -85,29 +92,19 @@ protected:
 	*/
 	EReadResult ReadInChunks(uint16* Buffer, const FString& ImagePath, int32 FrameId, const FIntPoint& Dim, int32 BufferSize);
 
-	/** 
-	 * Reads tiles from exr files in tile rows based on Tile region. If frame is pending for cancelation 
-	 * stops reading tiles at the current tile row.
-	*/
-	EReadResult ReadTiles(uint16* Buffer, const FString& ImagePath, int32 FrameId, const FIntRect& TileRegion, const FIntPoint& FullTexDimInTiles, const FIntPoint& TileDim, const int32 PixelSize);
-
 	/**
 	 * Get the size of the buffer needed to load in an image.
 	 * 
 	 * @param Dim Dimensions of the image.
 	 * @param NumChannels Number of channels in the image.
 	 */
-	static SIZE_T GetBufferSize(const FIntPoint& Dim, int32 NumChannels, bool bHasTiles, const FIntPoint& TileNum);
+	static SIZE_T GetBufferSize(const FIntPoint& Dim, int32 NumChannels, bool bHasTiles, const FIntPoint& TileNum, const bool bCustomExr);
 
-
-	TSharedPtr<IMediaTextureSampleConverter, ESPMode::ThreadSafe> CreateSampleConverter
-		( TArray<FStructuredBufferPoolItemSharedPtr>&& BufferDataArray
-		, const FIntPoint& FullResolution
-		, const FIntPoint& TileDim
-		, const TMap<int32, FIntRect>& Viewports
-		, const int32 NumChannels
-		, const int32 NumMipLevels
-		, const bool bHasTiles);
+	/**
+	* Creates Sample converter to be used by Media Texture Resource.
+	*/
+	static TSharedPtr<IMediaTextureSampleConverter, ESPMode::ThreadSafe> CreateSampleConverter
+		(TArray<FStructuredBufferPoolItemSharedPtr>&& BufferDataArray, TSharedPtr<FSampleConverterParameters> ConverterParams);
 
 public:
 
diff --git a/Engine/Plugins/Media/ImgMedia/Source/ImgMedia/Private/Readers/IImgMediaReader.h b/Engine/Plugins/Media/ImgMedia/Source/ImgMedia/Private/Readers/IImgMediaReader.h
--- a/Engine/Plugins/Media/ImgMedia/Source/ImgMedia/Private/Readers/IImgMediaReader.h
+++ b/Engine/Plugins/Media/ImgMedia/Source/ImgMedia/Private/Readers/IImgMediaReader.h
@@ -132,7 +132,7 @@ public:
 	 * maximum number of frames with as much efficiency as possible.
 	 *
 	 */
-	virtual void PreAllocateMemoryPool(int32 NumFrames, const FImgMediaFrameInfo& FrameInfo) {};
+	virtual void PreAllocateMemoryPool(int32 NumFrames, const FImgMediaFrameInfo& FrameInfo, const bool bCustomFormat) {};
 
 	/**
 	 * Used in case reader needs to do some processing once per frame.
