ImgMedia ExrReaderGpu: Added missing change from the previous CL.

diff --git a/Engine/Plugins/Media/ImgMedia/Source/ImgMedia/Private/Loader/ImgMediaLoader.cpp b/Engine/Plugins/Media/ImgMedia/Source/ImgMedia/Private/Loader/ImgMediaLoader.cpp
--- a/Engine/Plugins/Media/ImgMedia/Source/ImgMedia/Private/Loader/ImgMediaLoader.cpp
+++ b/Engine/Plugins/Media/ImgMedia/Source/ImgMedia/Private/Loader/ImgMediaLoader.cpp
@@ -930,7 +930,7 @@ void FImgMediaLoader::LoadSequence(const FString& SequencePath, const FFrameRate
 	// Giving our reader a chance to handle RAM allocation.
 	// Not all readers use this, only those that need to handle large files 
 	// or need to be as efficient as possible.
-	Reader->PreAllocateMemoryPool(NumFramesToLoad, FirstFrameInfo);
+	Reader->PreAllocateMemoryPool(NumFramesToLoad, FirstFrameInfo, FirstFrameInfo.FormatName == TEXT("EXR CUSTOM"));
 
 	Frames.Empty(NumFramesToLoad);
 
