ImgMedia ExrReaderGpu: Restricted custom exr reader to Windows platform.

diff --git a/Engine/Plugins/Media/ImgMedia/Source/ImgMedia/Private/Readers/ExrImgMediaReader.cpp b/Engine/Plugins/Media/ImgMedia/Source/ImgMedia/Private/Readers/ExrImgMediaReader.cpp
--- a/Engine/Plugins/Media/ImgMedia/Source/ImgMedia/Private/Readers/ExrImgMediaReader.cpp
+++ b/Engine/Plugins/Media/ImgMedia/Source/ImgMedia/Private/Readers/ExrImgMediaReader.cpp
@@ -81,6 +81,7 @@ FExrImgMediaReader::EReadResult FExrImgMediaReader::ReadTilesCustom
 	, const int32 PixelSize
 	, const bool bCustomExr)
 {
+#if defined(PLATFORM_WINDOWS) && PLATFORM_WINDOWS
 	EReadResult bResult = Success;
 
 	int64 CurrentBufferPos = 0;
@@ -130,8 +131,10 @@ FExrImgMediaReader::EReadResult FExrImgMediaReader::ReadTilesCustom
 	{
 		return Fail;
 	}
-
 	return bResult;
+#else
+	return Fail;
+#endif
 }
 
 bool FExrImgMediaReader::ReadFrame(int32 FrameId, const TMap<int32, FImgMediaTileSelection>& InMipTiles, TSharedPtr<FImgMediaFrame, ESPMode::ThreadSafe> OutFrame)
@@ -248,6 +251,7 @@ bool FExrImgMediaReader::ReadFrame(int32 FrameId, const TMap<int32, FImgMediaTil
 
 				if (OutFrame->Info.FormatName == TEXT("EXR CUSTOM"))
 				{
+#if defined(PLATFORM_WINDOWS) && PLATFORM_WINDOWS
 					FIntRect TileRegion = FIntRect(
 						(int32)CurrentTileSelection.TopLeftX,
 						(int32)CurrentTileSelection.TopLeftY,
@@ -265,6 +269,9 @@ bool FExrImgMediaReader::ReadFrame(int32 FrameId, const TMap<int32, FImgMediaTil
 					{
 						UE_LOG(LogImgMedia, Error, TEXT("Could not load %s"), *Image);
 					}
+#else
+					UE_LOG(LogImgMedia, Error, TEXT("Current platform doesn't support custom EXR file %s"), *Image);
+#endif
 				}
 				else
 				{
