ImgMedia: Measures bandwidth for the current frame.
#rnx

diff --git a/Engine/Plugins/Media/ImgMedia/Source/ImgMedia/Private/Loader/ImgMediaLoader.cpp b/Engine/Plugins/Media/ImgMedia/Source/ImgMedia/Private/Loader/ImgMediaLoader.cpp
--- a/Engine/Plugins/Media/ImgMedia/Source/ImgMedia/Private/Loader/ImgMediaLoader.cpp
+++ b/Engine/Plugins/Media/ImgMedia/Source/ImgMedia/Private/Loader/ImgMediaLoader.cpp
@@ -1333,7 +1333,8 @@ void FImgMediaLoader::AddEmptyFrame(int32 FrameNumber)
 	AddFrameToCache(FrameNumber, Frame);
 }
 
-void FImgMediaLoader::NotifyWorkComplete(FImgMediaLoaderWork& CompletedWork, int32 FrameNumber, const TSharedPtr<FImgMediaFrame, ESPMode::ThreadSafe>& Frame)
+void FImgMediaLoader::NotifyWorkComplete(FImgMediaLoaderWork& CompletedWork, int32 FrameNumber,
+	const TSharedPtr<FImgMediaFrame, ESPMode::ThreadSafe>& Frame, float WorkTime)
 {
 	FScopeLock Lock(&CriticalSection);
 
@@ -1344,6 +1345,13 @@ void FImgMediaLoader::NotifyWorkComplete(FImgMediaLoaderWork& CompletedWork, int
 	}
 
 	WorkPool.Push(&CompletedWork);
+
+	// Update bandwidth.
+	if ((Frame.IsValid()) && (WorkTime > 0.0f))
+	{
+		SIZE_T BytesLoaded = Frame->Info.UncompressedSize;
+		CurrentBandwidth = BytesLoaded / WorkTime;
+	}
 }
 
 void FImgMediaLoader::AddFrameToCache(int32 FrameNumber, const TSharedPtr<FImgMediaFrame, ESPMode::ThreadSafe>& Frame)
diff --git a/Engine/Plugins/Media/ImgMedia/Source/ImgMedia/Private/Loader/ImgMediaLoader.h b/Engine/Plugins/Media/ImgMedia/Source/ImgMedia/Private/Loader/ImgMediaLoader.h
--- a/Engine/Plugins/Media/ImgMedia/Source/ImgMedia/Private/Loader/ImgMediaLoader.h
+++ b/Engine/Plugins/Media/ImgMedia/Source/ImgMedia/Private/Loader/ImgMediaLoader.h
@@ -280,8 +280,10 @@ public:
 	 * @param CompletedWork The work item that completed.
 	 * @param FrameNumber Number of the frame that was read.
 	 * @param Frame The frame that was read, or nullptr if reading failed.
+	 * @param WorkTime How long to read this frame (in seconds).
 	 */
-	void NotifyWorkComplete(FImgMediaLoaderWork& CompletedWork, int32 FrameNumber, const TSharedPtr<FImgMediaFrame, ESPMode::ThreadSafe>& Frame);
+	void NotifyWorkComplete(FImgMediaLoaderWork& CompletedWork, int32 FrameNumber,
+		const TSharedPtr<FImgMediaFrame, ESPMode::ThreadSafe>& Frame, float WorkTime);
 
 	/**
 	 * Asynchronously request the image frame at the specified time.
@@ -555,4 +557,7 @@ private:
 
 	/** Settings for the smart cache. */
 	FImgMediaLoaderSmartCacheSettings SmartCacheSettings;
+
+	/** Number of bytes read per second from the last frame read. */
+	float CurrentBandwidth;
 };
diff --git a/Engine/Plugins/Media/ImgMedia/Source/ImgMedia/Private/Loader/ImgMediaLoaderWork.cpp b/Engine/Plugins/Media/ImgMedia/Source/ImgMedia/Private/Loader/ImgMediaLoaderWork.cpp
--- a/Engine/Plugins/Media/ImgMedia/Source/ImgMedia/Private/Loader/ImgMediaLoaderWork.cpp
+++ b/Engine/Plugins/Media/ImgMedia/Source/ImgMedia/Private/Loader/ImgMediaLoaderWork.cpp
@@ -49,13 +49,14 @@ void FImgMediaLoaderWork::Abandon()
 {
 	SCOPE_CYCLE_COUNTER(STAT_ImgMedia_LoaderAbandonWork);
 
-	Finalize(nullptr);
+	Finalize(nullptr, 0.0f);
 }
 
 
 void FImgMediaLoaderWork::DoThreadedWork()
 {
 	TSharedPtr<FImgMediaFrame, ESPMode::ThreadSafe> Frame;
+	float WorkTime = 0.0f;
 	UE_LOG(LogImgMedia, Verbose, TEXT("Loader %p: Starting to read %i"), this, FrameNumber);
 
 	if (FrameNumber == INDEX_NONE)
@@ -76,6 +77,7 @@ void FImgMediaLoaderWork::DoThreadedWork()
 			Frame = ExistingFrame;
 		}
 
+		double StartTime = FPlatformTime::Seconds();
 		if(!Reader->ReadFrame(FrameNumber, MipTiles, Frame))
 		{
 			if (ExistingFrame == nullptr)
@@ -83,18 +85,22 @@ void FImgMediaLoaderWork::DoThreadedWork()
 				Frame.Reset();
 			}
 		}
+		else
+		{
+			WorkTime = FPlatformTime::Seconds() - StartTime;
+		}
 	}
 
 	SCOPE_CYCLE_COUNTER(STAT_ImgMedia_LoaderFinalizeWork);
 
-	Finalize(Frame);
+	Finalize(Frame, WorkTime);
 }
 
 
 /* FImgMediaLoaderWork implementation
  *****************************************************************************/
 
-void FImgMediaLoaderWork::Finalize(TSharedPtr<FImgMediaFrame, ESPMode::ThreadSafe> Frame)
+void FImgMediaLoaderWork::Finalize(TSharedPtr<FImgMediaFrame, ESPMode::ThreadSafe> Frame, float WorkTime)
 {
 	ExistingFrame.Reset();
 
@@ -102,7 +108,7 @@ void FImgMediaLoaderWork::Finalize(TSharedPtr<FImgMediaFrame, ESPMode::ThreadSaf
 
 	if (Owner.IsValid())
 	{
-		Owner->NotifyWorkComplete(*this, FrameNumber, Frame);
+		Owner->NotifyWorkComplete(*this, FrameNumber, Frame, WorkTime);
 	}
 	else
 	{
diff --git a/Engine/Plugins/Media/ImgMedia/Source/ImgMedia/Private/Loader/ImgMediaLoaderWork.h b/Engine/Plugins/Media/ImgMedia/Source/ImgMedia/Private/Loader/ImgMediaLoaderWork.h
--- a/Engine/Plugins/Media/ImgMedia/Source/ImgMedia/Private/Loader/ImgMediaLoaderWork.h
+++ b/Engine/Plugins/Media/ImgMedia/Source/ImgMedia/Private/Loader/ImgMediaLoaderWork.h
@@ -54,8 +54,9 @@ protected:
 	 * Notify the owner of this work item, or self destruct.
 	 *
 	 * @param Frame The frame that was loaded by this work item.
+	 * @param WorkTime How long to read this frame (in seconds).
 	 */
-	void Finalize(TSharedPtr<FImgMediaFrame, ESPMode::ThreadSafe> Frame);
+	void Finalize(TSharedPtr<FImgMediaFrame, ESPMode::ThreadSafe> Frame, float WorkTime);
 
 private:
 
