ImgMedia: Editor module keeps a list of players.

#rnx

diff --git a/Engine/Plugins/Media/ImgMedia/Source/ImgMedia/Private/ImgMediaModule.cpp b/Engine/Plugins/Media/ImgMedia/Source/ImgMedia/Private/ImgMediaModule.cpp
--- a/Engine/Plugins/Media/ImgMedia/Source/ImgMedia/Private/ImgMediaModule.cpp
+++ b/Engine/Plugins/Media/ImgMedia/Source/ImgMedia/Private/ImgMediaModule.cpp
@@ -116,7 +116,10 @@ public:
 			InitGlobalCache();
 		}
 
-		return MakeShared<FImgMediaPlayer, ESPMode::ThreadSafe>(EventSink, Scheduler.ToSharedRef(), GlobalCache.ToSharedRef());
+		TSharedPtr<FImgMediaPlayer, ESPMode::ThreadSafe> Player = MakeShared<FImgMediaPlayer, ESPMode::ThreadSafe>(EventSink, Scheduler.ToSharedRef(), GlobalCache.ToSharedRef());
+		OnImgMediaPlayerCreated.Broadcast(Player);
+
+		return Player;
 	}
 
 public:
diff --git a/Engine/Plugins/Media/ImgMedia/Source/ImgMedia/Public/IImgMediaModule.h b/Engine/Plugins/Media/ImgMedia/Source/ImgMedia/Public/IImgMediaModule.h
--- a/Engine/Plugins/Media/ImgMedia/Source/ImgMedia/Public/IImgMediaModule.h
+++ b/Engine/Plugins/Media/ImgMedia/Source/ImgMedia/Public/IImgMediaModule.h
@@ -3,13 +3,17 @@
 #pragma once
 
 #include "CoreMinimal.h"
+#include "Delegates/Delegate.h"
 #include "Templates/SharedPointer.h"
 #include "Modules/ModuleInterface.h"
 
 class FImgMediaGlobalCache;
+class FImgMediaPlayer;
 class IMediaEventSink;
 class IMediaPlayer;
 
+/** Callback when a player gets created. */
+DECLARE_MULTICAST_DELEGATE_OneParam(FOnImgMediaPlayerCreated, const TSharedPtr<FImgMediaPlayer>&);
 
 /**
  * Interface for the ImgMedia module.
@@ -32,6 +36,9 @@ public:
 	/** Virtual destructor. */
 	virtual ~IImgMediaModule() { }
 
+	/** Add to this callback to get called whenever a player is created. */
+	FOnImgMediaPlayerCreated OnImgMediaPlayerCreated;
+
 	/**
 	 * Call this to get the global cache.
 	 * 
diff --git a/Engine/Plugins/Media/ImgMedia/Source/ImgMediaEditor/Private/ImgMediaEditorModule.cpp b/Engine/Plugins/Media/ImgMedia/Source/ImgMediaEditor/Private/ImgMediaEditorModule.cpp
--- a/Engine/Plugins/Media/ImgMedia/Source/ImgMediaEditor/Private/ImgMediaEditorModule.cpp
+++ b/Engine/Plugins/Media/ImgMedia/Source/ImgMediaEditor/Private/ImgMediaEditorModule.cpp
@@ -1,13 +1,14 @@
 // Copyright Epic Games, Inc. All Rights Reserved.
 
 #include "ImgMediaEditorModule.h"
-#include "Modules/ModuleInterface.h"
 #include "Modules/ModuleManager.h"
 
 #include "AssetToolsModule.h"
 #include "AssetTools/ImgMediaSourceActions.h"
 #include "Customizations/ImgMediaSourceCustomization.h"
 #include "IAssetTools.h"
+#include "ImgMediaEditorModule.h"
+#include "IImgMediaModule.h"
 #include "ImgMediaSource.h"
 #include "PropertyEditorModule.h"
 #include "UObject/NameTypes.h"
@@ -27,9 +28,14 @@ static const FName ImgMediaCacheTabName(TEXT("ImgMediaCache"));
  * Implements the ImgMediaEditor module.
  */
 class FImgMediaEditorModule
-	: public IModuleInterface
+	: public IImgMediaEditorModule
 {
 public:
+	//~ IImgMediaEditorModule interface
+	const TArray<TWeakPtr<FImgMediaPlayer>>& GetMediaPlayers()
+	{
+		return MediaPlayers;
+	}
 
 	//~ IModuleInterface interface
 
@@ -38,6 +44,13 @@ public:
 		RegisterCustomizations();
 		RegisterAssetTools();
 		RegisterTabSpawners();
+
+		IImgMediaModule* ImgMediaModule = FModuleManager::LoadModulePtr<IImgMediaModule>("ImgMedia");
+		if (ImgMediaModule != nullptr)
+		{
+			ImgMediaModule->OnImgMediaPlayerCreated.AddRaw(this, &FImgMediaEditorModule::OnImgMediaPlayerCreated);
+
+		}
 	}
 
 	virtual void ShutdownModule() override
@@ -150,6 +163,30 @@ protected:
 			];
 	}
 
+	void OnImgMediaPlayerCreated(const TSharedPtr<FImgMediaPlayer>& Player)
+	{
+		// Try and replace an expired player.
+		bool bIsAdded = false;
+		for (TWeakPtr<FImgMediaPlayer>& PlayerPointer : MediaPlayers)
+		{
+			if (PlayerPointer.IsValid() == false)
+			{
+				PlayerPointer = Player;
+				bIsAdded = true;
+				break;
+			}
+		}
+
+		// If we were not able to add it, just add it now.
+		if (bIsAdded == false)
+		{
+			MediaPlayers.Add(Player);
+		}
+
+		// Send out the message.
+		OnImgMediaEditorPlayersUpdated.Broadcast();
+	}
+
 private:
 
 	/** Customization name to avoid reusing staticstruct during shutdown. */
@@ -160,6 +197,9 @@ private:
 
 	/** The collection of registered asset type actions. */
 	TArray<TSharedRef<IAssetTypeActions>> RegisteredAssetTypeActions;
+
+	/** Array of all our players. */
+	TArray<TWeakPtr<FImgMediaPlayer>> MediaPlayers;
 };
 
 FLazyName FImgMediaEditorModule::ImgMediaProcessImagesTabName(TEXT("ImgMediaProcessImages"));
diff --git a/Engine/Plugins/Media/ImgMedia/Source/ImgMediaEditor/Private/ImgMediaEditorModule.h b/Engine/Plugins/Media/ImgMedia/Source/ImgMediaEditor/Private/ImgMediaEditorModule.h
--- a/Engine/Plugins/Media/ImgMedia/Source/ImgMediaEditor/Private/ImgMediaEditorModule.h
+++ b/Engine/Plugins/Media/ImgMedia/Source/ImgMediaEditor/Private/ImgMediaEditorModule.h
@@ -2,7 +2,32 @@
 
 #pragma once
 
+#include "Delegates/Delegate.h"
 #include "Logging/LogMacros.h"
+#include "Modules/ModuleInterface.h"
+
+class FImgMediaPlayer;
 
 /** Log category for this module. */
 DECLARE_LOG_CATEGORY_EXTERN(LogImgMediaEditor, Log, All);
+
+/** Callback for when the list of players is updated. */
+DECLARE_MULTICAST_DELEGATE(FOnImgMediaEditorPlayersUpdated);
+
+
+/**
+* Interface for the ImgMediaEditor module.
+*/
+class IMGMEDIAEDITOR_API IImgMediaEditorModule
+	: public IModuleInterface
+{
+public:
+
+	/**
+	 * Call this to get a list of all the players. *
+	 */
+	virtual const TArray<TWeakPtr<FImgMediaPlayer>>& GetMediaPlayers() = 0;
+
+	/** Add to this to get a callback when the list of players gets updated. */
+	FOnImgMediaEditorPlayersUpdated OnImgMediaEditorPlayersUpdated;
+};
