ImgMedia: Fixed process images not handling RGB files correctly.

#rnx

diff --git a/Engine/Plugins/Media/ImgMedia/Source/ImgMediaEditor/Private/Widgets/SImgMediaProcessImages.cpp b/Engine/Plugins/Media/ImgMedia/Source/ImgMediaEditor/Private/Widgets/SImgMediaProcessImages.cpp
--- a/Engine/Plugins/Media/ImgMedia/Source/ImgMediaEditor/Private/Widgets/SImgMediaProcessImages.cpp
+++ b/Engine/Plugins/Media/ImgMedia/Source/ImgMediaEditor/Private/Widgets/SImgMediaProcessImages.cpp
@@ -307,7 +307,11 @@ void SImgMediaProcessImages::ProcessImageCustom(TSharedPtr<IImageWrapper>& InIma
 	// ImageWrapper always returns an alpha channel, so make sure we really have one.
 	if ((DestNumChannels == 4) && (bHasAlphaChannel == false))
 	{
+		// Remove the alpha channel as its not needed.
+		RemoveAlphaChannel(RawData);
+		NumChannels = 3;
 		DestNumChannels = 3;
+		BytesPerPixel = BytesPerPixelPerChannel * NumChannels;
 	}
 
 	TArray64<uint8> TileBuffer;
@@ -479,18 +483,20 @@ void SImgMediaProcessImages::ProcessImageCustom(TSharedPtr<IImageWrapper>& InIma
 
 		// Write to EXR.
 		Stride.Y = MipWidth * BytesPerPixel;
-		
+		int64 BufferOffset = 0;
+		int64 SingleBufferOffset = MipWidth * BytesPerPixelPerChannel;
 		if (DestNumChannels == 4)
 		{
 			OutFile.UpdateFrameBufferChannel(AChannelName, CurrentBuffer, Stride);
+			BufferOffset += SingleBufferOffset;
 		}
 
-		OutFile.UpdateFrameBufferChannel(BChannelName,
-			CurrentBuffer + MipWidth * BytesPerPixelPerChannel, Stride);
-		OutFile.UpdateFrameBufferChannel(GChannelName,
-			CurrentBuffer + MipWidth * 2 * BytesPerPixelPerChannel, Stride);
-		OutFile.UpdateFrameBufferChannel(RChannelName,
-			CurrentBuffer + MipWidth * 3 * BytesPerPixelPerChannel, Stride);
+		OutFile.UpdateFrameBufferChannel(BChannelName, CurrentBuffer + BufferOffset, Stride);
+		BufferOffset += SingleBufferOffset;
+		OutFile.UpdateFrameBufferChannel(GChannelName, CurrentBuffer + BufferOffset, Stride);
+		BufferOffset += SingleBufferOffset;
+		OutFile.UpdateFrameBufferChannel(RChannelName, CurrentBuffer + BufferOffset, Stride);
+		BufferOffset += SingleBufferOffset;
 
 		OutFile.SetFrameBuffer();
 
@@ -507,6 +513,29 @@ void SImgMediaProcessImages::ProcessImageCustom(TSharedPtr<IImageWrapper>& InIma
 #endif // IMGMEDIAEDITOR_EXR_SUPPORTED_PLATFORM
 }
 
+void SImgMediaProcessImages::RemoveAlphaChannel(TArray64<uint8>& Buffer)
+{
+	int32 BytesPerPixelPerChannel = 2;
+	int64 BufferSize = Buffer.Num() / BytesPerPixelPerChannel;
+	uint16* BufferPtr = (uint16*)(Buffer.GetData());
+	
+	// Loop through the buffer.
+	int64 OutIndex = 0;
+	for (int64 Index = 0; Index < BufferSize; ++Index)
+	{
+		// Skip every fourth channel (i.e. the alpha channel).
+		if ((Index & 0x3) != 3)
+		{
+			// Copy the data in place.
+			BufferPtr[OutIndex] = BufferPtr[Index];
+			OutIndex++;
+		}
+	}
+
+	// Don't bother shrinking as its just a waste and extra work.
+	Buffer.SetNum((BufferSize * 3) / 4, false);
+}
+
 void SImgMediaProcessImages::TileData(uint8* SourceData, TArray64<uint8>& DestArray,
 	int32 SourceWidth, int32 SourceHeight, int32 DestWidth, int32 DestHeight,
 	int32 NumTilesX, int32 NumTilesY,
diff --git a/Engine/Plugins/Media/ImgMedia/Source/ImgMediaEditor/Private/Widgets/SImgMediaProcessImages.h b/Engine/Plugins/Media/ImgMedia/Source/ImgMediaEditor/Private/Widgets/SImgMediaProcessImages.h
--- a/Engine/Plugins/Media/ImgMedia/Source/ImgMediaEditor/Private/Widgets/SImgMediaProcessImages.h
+++ b/Engine/Plugins/Media/ImgMedia/Source/ImgMediaEditor/Private/Widgets/SImgMediaProcessImages.h
@@ -76,6 +76,15 @@ private:
 		int32 InTileWidth, int32 InTileHeight, int32 InTileBorder, 
 		bool bInEnableMips, bool bHasAlphaChannel, const FString& InName);
 
+	/**
+	 * Removes the alpha channel from a buffer.
+	 * This will be done in place.
+	 * Assumes 2 bytes per channel, 4 channels, alpha is the last channel.
+	 * 
+	 * @param Buffer			Buffer to remove alpha from.
+	 */
+	void RemoveAlphaChannel(TArray64<uint8>& Buffer);
+
 	/**
 	 * Creates tiles from a source and outputs it to a destination.
 	 *
