ImgMedia: Added bandwidth window.

#rnx

diff --git a/Engine/Plugins/Media/ImgMedia/Source/ImgMedia/Private/Loader/ImgMediaLoader.h b/Engine/Plugins/Media/ImgMedia/Source/ImgMedia/Private/Loader/ImgMediaLoader.h
--- a/Engine/Plugins/Media/ImgMedia/Source/ImgMedia/Private/Loader/ImgMediaLoader.h
+++ b/Engine/Plugins/Media/ImgMedia/Source/ImgMedia/Private/Loader/ImgMediaLoader.h
@@ -75,6 +75,9 @@ public:
 	 */
 	uint64 GetBitRate() const;
 
+	/** Get the number of bytes read per second from the last frame read. */
+	float GetCurrentBandwidth() const { return CurrentBandwidth; }
+
 	/**
 	 * Get the time ranges of frames that are being loaded right now.
 	 *
diff --git a/Engine/Plugins/Media/ImgMedia/Source/ImgMedia/Private/Player/ImgMediaPlayer.h b/Engine/Plugins/Media/ImgMedia/Source/ImgMedia/Private/Player/ImgMediaPlayer.h
--- a/Engine/Plugins/Media/ImgMedia/Source/ImgMedia/Private/Player/ImgMediaPlayer.h
+++ b/Engine/Plugins/Media/ImgMedia/Source/ImgMedia/Private/Player/ImgMediaPlayer.h
@@ -43,6 +43,9 @@ public:
 	/** Virtual destructor. */
 	virtual ~FImgMediaPlayer();
 
+	/** Get the loader. */
+	TSharedPtr<FImgMediaLoader, ESPMode::ThreadSafe> GetLoader() { return Loader; }
+
 public:
 
 	//~ IMediaPlayer interface
diff --git a/Engine/Plugins/Media/ImgMedia/Source/ImgMediaEditor/ImgMediaEditor.Build.cs b/Engine/Plugins/Media/ImgMedia/Source/ImgMediaEditor/ImgMediaEditor.Build.cs
--- a/Engine/Plugins/Media/ImgMedia/Source/ImgMediaEditor/ImgMediaEditor.Build.cs
+++ b/Engine/Plugins/Media/ImgMedia/Source/ImgMediaEditor/ImgMediaEditor.Build.cs
@@ -26,6 +26,7 @@ namespace UnrealBuildTool.Rules
 
 			PrivateIncludePaths.AddRange(
 				new string[] {
+					"ImgMedia/Private",
 					"ImgMediaEditor/Private",
 					"ImgMediaEditor/Private/Customizations",
 					"ImgMediaEditor/Private/Factories",
diff --git a/Engine/Plugins/Media/ImgMedia/Source/ImgMediaEditor/Private/ImgMediaEditorModule.cpp b/Engine/Plugins/Media/ImgMedia/Source/ImgMediaEditor/Private/ImgMediaEditorModule.cpp
--- a/Engine/Plugins/Media/ImgMedia/Source/ImgMediaEditor/Private/ImgMediaEditorModule.cpp
+++ b/Engine/Plugins/Media/ImgMedia/Source/ImgMediaEditor/Private/ImgMediaEditorModule.cpp
@@ -13,6 +13,7 @@
 #include "PropertyEditorModule.h"
 #include "UObject/NameTypes.h"
 #include "Widgets/Docking/SDockTab.h"
+#include "Widgets/SImgMediaBandwidth.h"
 #include "Widgets/SImgMediaCache.h"
 #include "Widgets/SImgMediaProcessImages.h"
 #include "WorkspaceMenuStructure.h"
@@ -22,6 +23,7 @@
 
 DEFINE_LOG_CATEGORY(LogImgMediaEditor);
 
+static const FName ImgMediaBandwidthTabName(TEXT("ImgMediaBandwidth"));
 static const FName ImgMediaCacheTabName(TEXT("ImgMediaCache"));
 
 /**
@@ -119,6 +121,14 @@ protected:
 			FSlateIcon(FEditorStyle::GetStyleSetName(), "SequenceRecorder.TabIcon"),
 			true);
 
+		// Add bandwidth tab.
+		FGlobalTabmanager::Get()->RegisterNomadTabSpawner(ImgMediaBandwidthTabName,
+			FOnSpawnTab::CreateStatic(&FImgMediaEditorModule::SpawnBandwidthTab))
+			.SetGroup(MediaBrowserGroup)
+			.SetDisplayName(LOCTEXT("ImgMediaBandwidthTabTitle", "Bandwidth"))
+			.SetTooltipText(LOCTEXT("ImgMediaBandwidthTooltipText", "Open the bandwidth tab."))
+			.SetIcon(FSlateIcon(FEditorStyle::GetStyleSetName(), "SequenceRecorder.TabIcon"));
+
 		// Add cache tab.
 		FGlobalTabmanager::Get()->RegisterNomadTabSpawner(ImgMediaCacheTabName,
 			FOnSpawnTab::CreateStatic(&FImgMediaEditorModule::SpawnCacheTab))
@@ -142,9 +152,19 @@ protected:
 		{
 			FGlobalTabmanager::Get()->UnregisterNomadTabSpawner(ImgMediaProcessImagesTabName);
 			FGlobalTabmanager::Get()->UnregisterNomadTabSpawner(ImgMediaCacheTabName);
+			FGlobalTabmanager::Get()->UnregisterNomadTabSpawner(ImgMediaBandwidthTabName);
 		}
 	}
 
+	static TSharedRef<SDockTab> SpawnBandwidthTab(const FSpawnTabArgs& SpawnTabArgs)
+	{
+		return SNew(SDockTab)
+			.TabRole(ETabRole::NomadTab)
+			[
+				SNew(SImgMediaBandwidth)
+			];
+	}
+
 	static TSharedRef<SDockTab> SpawnCacheTab(const FSpawnTabArgs& SpawnTabArgs)
 	{
 		return SNew(SDockTab)
diff --git a/Engine/Plugins/Media/ImgMedia/Source/ImgMediaEditor/Private/Widgets/SImgMediaBandwidth.cpp b/Engine/Plugins/Media/ImgMedia/Source/ImgMediaEditor/Private/Widgets/SImgMediaBandwidth.cpp
new file mode 100644
--- /dev/null
+++ b/Engine/Plugins/Media/ImgMedia/Source/ImgMediaEditor/Private/Widgets/SImgMediaBandwidth.cpp
@@ -0,0 +1,145 @@
+// Copyright Epic Games, Inc. All Rights Reserved.
+
+#include "Widgets/SImgMediaBandwidth.h"
+#include "DetailLayoutBuilder.h"
+#include "IImgMediaModule.h"
+#include "ImgMediaEditorModule.h"
+#include "Loader/ImgMediaLoader.h"
+#include "Modules/ModuleManager.h"
+#include "Player/ImgMediaPlayer.h"
+#include "SlateOptMacros.h"
+#include "UObject/UObjectIterator.h"
+#include "Widgets/Input/SButton.h"
+#include "Widgets/Layout/SGridPanel.h"
+#include "Widgets/Layout/SScrollBox.h"
+#include "Widgets/Layout/SSplitter.h"
+#include "Widgets/Text/STextBlock.h"
+
+#define LOCTEXT_NAMESPACE "ImgMediaBandwidth"
+
+BEGIN_SLATE_FUNCTION_BUILD_OPTIMIZATION
+void SImgMediaBandwidth::Construct(const FArguments& InArgs)
+{
+	ChildSlot
+	[
+		SNew(SVerticalBox)
+
+		// Bandwidth label.
+		+ SVerticalBox::Slot()
+			.AutoHeight()
+			.Padding(0, 5, 0, 0)
+		[
+			SNew(STextBlock)
+				.Text(LOCTEXT("Bandwidth", "Bandwidth"))
+		]
+
+		// Container for all players.
+		+ SVerticalBox::Slot()
+			.AutoHeight()
+			[
+				SAssignNew(PlayersContainer, SVerticalBox)
+			]
+		
+	];
+
+	// Update players.
+	RefreshPlayersContainer();
+
+	// Get notified when a player is created.
+	IImgMediaEditorModule* ImgMediaEditorModule = FModuleManager::LoadModulePtr<IImgMediaEditorModule>("ImgMediaEditor");
+	if (ImgMediaEditorModule != nullptr)
+	{
+		ImgMediaEditorModule->OnImgMediaEditorPlayersUpdated.AddSP(this, &SImgMediaBandwidth::RefreshPlayersContainer);
+	}
+
+}
+END_SLATE_FUNCTION_BUILD_OPTIMIZATION
+
+void SImgMediaBandwidth::Tick(const FGeometry& AllottedGeometry, const double InCurrentTime, const float InDeltaTime)
+{
+	// Call parent.
+	SCompoundWidget::Tick(AllottedGeometry, InCurrentTime, InDeltaTime);
+
+	// Loop over all players.
+	for (const FPlayerInfo& PlayerInfo : PlayerInfos)
+	{
+		TSharedPtr< FImgMediaPlayer> Player = PlayerInfo.Player.Pin();
+		if (Player != nullptr)
+		{
+			// Update widgets with info.
+			// Set URL.
+			PlayerInfo.UrlTextBlock->SetText(FText::Format(
+				LOCTEXT("Url", "URL: {0}"),
+				FText::FromString(Player->GetUrl())));
+
+			// Get the loader.
+			TSharedPtr<FImgMediaLoader, ESPMode::ThreadSafe> Loader = Player->GetLoader();
+			if (Loader.IsValid())
+			{
+				// Set bandwidth.
+				float Bandwidth = Loader->GetCurrentBandwidth() / (1024.0f * 1024.0f);
+				PlayerInfo.BandwidthTextBlock->SetText(FText::Format(
+					LOCTEXT("Current", "Current: {0} MB/s"),
+					FText::AsNumber(Bandwidth)));
+			}
+		}
+	}
+}
+
+void SImgMediaBandwidth::RefreshPlayersContainer()
+{
+	// Reset our players.
+	PlayersContainer->ClearChildren();
+	PlayerInfos.Reset();
+
+	IImgMediaEditorModule* ImgMediaEditorModule = FModuleManager::LoadModulePtr<IImgMediaEditorModule>("ImgMediaEditor");
+	if (ImgMediaEditorModule != nullptr)
+	{
+		const TArray<TWeakPtr<FImgMediaPlayer>>& MediaPlayers = ImgMediaEditorModule->GetMediaPlayers();
+
+		// Loop over all players.
+		for (TWeakPtr<FImgMediaPlayer> PlayerWeakPtr : MediaPlayers)
+		{
+			TSharedPtr< FImgMediaPlayer> Player = PlayerWeakPtr.Pin();
+			if (Player != nullptr)
+			{
+				FPlayerInfo PlayerInfo;
+				PlayerInfo.Player = Player;
+				PlayersContainer->AddSlot()
+					.AutoHeight()
+					.Padding(0, 5, 0, 0)
+					[
+						SNew(SVerticalBox)
+
+						// Add URL.
+						+ SVerticalBox::Slot()
+							.AutoHeight()
+							.Padding(0, 5, 0, 0)
+							[
+								SAssignNew(PlayerInfo.UrlTextBlock, STextBlock)
+							]
+
+						// Add bandwidth.
+						+ SVerticalBox::Slot()
+							.AutoHeight()
+							.Padding(0, 5, 0, 0)
+							[
+								SNew(SHorizontalBox)
+
+								+ SHorizontalBox::Slot()
+									.AutoWidth()
+									[
+										SAssignNew(PlayerInfo.BandwidthTextBlock, STextBlock)
+									]
+							]
+
+					];
+
+				// Add this to the list.
+				PlayerInfos.Add(PlayerInfo);
+			}
+		}
+	}
+}
+
+#undef LOCTEXT_NAMESPACE
diff --git a/Engine/Plugins/Media/ImgMedia/Source/ImgMediaEditor/Private/Widgets/SImgMediaBandwidth.h b/Engine/Plugins/Media/ImgMedia/Source/ImgMediaEditor/Private/Widgets/SImgMediaBandwidth.h
new file mode 100644
--- /dev/null
+++ b/Engine/Plugins/Media/ImgMedia/Source/ImgMediaEditor/Private/Widgets/SImgMediaBandwidth.h
@@ -0,0 +1,47 @@
+// Copyright Epic Games, Inc. All Rights Reserved.
+
+#pragma once
+
+#include "CoreMinimal.h"
+#include "Input/Reply.h"
+#include "Widgets/SCompoundWidget.h"
+
+class FImgMediaPlayer;
+class STextBlock;
+class SVerticalBox;
+
+/**
+ * SImgMediaCache manages caching for image sequences.
+ */
+class SImgMediaBandwidth : public SCompoundWidget
+{
+public:
+	SLATE_BEGIN_ARGS(SImgMediaBandwidth){}
+	SLATE_END_ARGS()
+
+	void Construct(const FArguments& InArgs);
+
+	//~ SWidget interface
+	virtual void Tick(const FGeometry& AllottedGeometry, const double InCurrentTime, const float InDeltaTime) override;
+
+private:
+
+	/**
+	 * Update our list of players and the widgets accordingly.
+	 */
+	void RefreshPlayersContainer();
+
+	/** Info for a single player. */
+	struct FPlayerInfo
+	{
+		TWeakPtr<FImgMediaPlayer> Player;
+		TSharedPtr<STextBlock> UrlTextBlock;
+		TSharedPtr<STextBlock> BandwidthTextBlock;
+	};
+
+	/** Array of all players. */
+	TArray<FPlayerInfo> PlayerInfos;
+	/** Widget to hold all widgets for all the players. */
+	TSharedPtr<SVerticalBox> PlayersContainer;
+
+};
