ImgMedia: Added minimum bandwidth to bandwidth window.

#rnx

diff --git a/Engine/Plugins/Media/ImgMedia/Source/ImgMedia/Private/Loader/ImgMediaLoader.cpp b/Engine/Plugins/Media/ImgMedia/Source/ImgMedia/Private/Loader/ImgMediaLoader.cpp
--- a/Engine/Plugins/Media/ImgMedia/Source/ImgMedia/Private/Loader/ImgMediaLoader.cpp
+++ b/Engine/Plugins/Media/ImgMedia/Source/ImgMedia/Private/Loader/ImgMediaLoader.cpp
@@ -872,6 +872,14 @@ void FImgMediaLoader::LoadSequence(const FString& SequencePath, const FFrameRate
 	SequenceDuration = FrameNumberToTime(GetNumImages());
 	SIZE_T UncompressedSize = FirstFrameInfo.UncompressedSize;
 
+	// Get minimum bandwidth needed for mip level 0.
+	MinBandwidthForMipLevel0 = UncompressedSize;
+	float FrameTime = SequenceFrameRate.AsInterval();
+	if (FrameTime > 0)
+	{
+		MinBandwidthForMipLevel0 /= FrameTime;
+	}
+
 #if 0
 	// todo: enable once both the CPU & GPU reader use the same tile count logic.
 	if (FirstFrameInfo.bHasTiles)
diff --git a/Engine/Plugins/Media/ImgMedia/Source/ImgMedia/Private/Loader/ImgMediaLoader.h b/Engine/Plugins/Media/ImgMedia/Source/ImgMedia/Private/Loader/ImgMediaLoader.h
--- a/Engine/Plugins/Media/ImgMedia/Source/ImgMedia/Private/Loader/ImgMediaLoader.h
+++ b/Engine/Plugins/Media/ImgMedia/Source/ImgMedia/Private/Loader/ImgMediaLoader.h
@@ -77,6 +77,8 @@ public:
 
 	/** Get the number of bytes read per second from the last frame read. */
 	float GetCurrentBandwidth() const { return CurrentBandwidth; }
+	/** Get the minimum number of bytes per second needed to read mip level 0.. */
+	float GetMinBandwidthForMipLevel0() const { return MinBandwidthForMipLevel0; }
 
 	/**
 	 * Get the time ranges of frames that are being loaded right now.
@@ -563,4 +565,6 @@ private:
 
 	/** Number of bytes read per second from the last frame read. */
 	float CurrentBandwidth;
+	/** Minimum number of bytes per second needed to read mip level 0. */
+	float MinBandwidthForMipLevel0;
 };
diff --git a/Engine/Plugins/Media/ImgMedia/Source/ImgMediaEditor/Private/Widgets/SImgMediaBandwidth.cpp b/Engine/Plugins/Media/ImgMedia/Source/ImgMediaEditor/Private/Widgets/SImgMediaBandwidth.cpp
--- a/Engine/Plugins/Media/ImgMedia/Source/ImgMediaEditor/Private/Widgets/SImgMediaBandwidth.cpp
+++ b/Engine/Plugins/Media/ImgMedia/Source/ImgMediaEditor/Private/Widgets/SImgMediaBandwidth.cpp
@@ -76,11 +76,17 @@ void SImgMediaBandwidth::Tick(const FGeometry& AllottedGeometry, const double In
 			TSharedPtr<FImgMediaLoader, ESPMode::ThreadSafe> Loader = Player->GetLoader();
 			if (Loader.IsValid())
 			{
-				// Set bandwidth.
+				// Set cuurrent bandwidth.
 				float Bandwidth = Loader->GetCurrentBandwidth() / (1024.0f * 1024.0f);
-				PlayerInfo.BandwidthTextBlock->SetText(FText::Format(
+				PlayerInfo.CurrentBandwidthTextBlock->SetText(FText::Format(
 					LOCTEXT("Current", "Current: {0} MB/s"),
 					FText::AsNumber(Bandwidth)));
+
+				// Set minimum bandwidth.
+				Bandwidth = Loader->GetMinBandwidthForMipLevel0() / (1024.0f * 1024.0f);
+				PlayerInfo.MinimumBandwidthTextBlock->SetText(FText::Format(
+					LOCTEXT("Minimum", "Minimum: {0} MB/s"),
+					FText::AsNumber(Bandwidth)));
 			}
 		}
 	}
@@ -127,9 +133,15 @@ void SImgMediaBandwidth::RefreshPlayersContainer()
 								SNew(SHorizontalBox)
 
 								+ SHorizontalBox::Slot()
-									.AutoWidth()
+									.FillWidth(1.0f)
+									[
+										SAssignNew(PlayerInfo.CurrentBandwidthTextBlock, STextBlock)
+									]
+
+								+ SHorizontalBox::Slot()
+									.FillWidth(1.0f)
 									[
-										SAssignNew(PlayerInfo.BandwidthTextBlock, STextBlock)
+										SAssignNew(PlayerInfo.MinimumBandwidthTextBlock, STextBlock)
 									]
 							]
 
diff --git a/Engine/Plugins/Media/ImgMedia/Source/ImgMediaEditor/Private/Widgets/SImgMediaBandwidth.h b/Engine/Plugins/Media/ImgMedia/Source/ImgMediaEditor/Private/Widgets/SImgMediaBandwidth.h
--- a/Engine/Plugins/Media/ImgMedia/Source/ImgMediaEditor/Private/Widgets/SImgMediaBandwidth.h
+++ b/Engine/Plugins/Media/ImgMedia/Source/ImgMediaEditor/Private/Widgets/SImgMediaBandwidth.h
@@ -36,7 +36,8 @@ private:
 	{
 		TWeakPtr<FImgMediaPlayer> Player;
 		TSharedPtr<STextBlock> UrlTextBlock;
-		TSharedPtr<STextBlock> BandwidthTextBlock;
+		TSharedPtr<STextBlock> CurrentBandwidthTextBlock;
+		TSharedPtr<STextBlock> MinimumBandwidthTextBlock;
 	};
 
 	/** Array of all players. */
