ImgMedia: Added some profiling to process images.
#rnx

diff --git a/Engine/Plugins/Media/ImgMedia/Source/ImgMediaEditor/Private/Widgets/SImgMediaProcessImages.cpp b/Engine/Plugins/Media/ImgMedia/Source/ImgMediaEditor/Private/Widgets/SImgMediaProcessImages.cpp
--- a/Engine/Plugins/Media/ImgMedia/Source/ImgMediaEditor/Private/Widgets/SImgMediaProcessImages.cpp
+++ b/Engine/Plugins/Media/ImgMedia/Source/ImgMediaEditor/Private/Widgets/SImgMediaProcessImages.cpp
@@ -15,6 +15,7 @@
 #include "Misc/FileHelper.h"
 #include "Misc/Paths.h"
 #include "Modules/ModuleManager.h"
+#include "ProfilingDebugging/CpuProfilerTrace.h"
 #include "PropertyEditorModule.h"
 #include "SlateOptMacros.h"
 #include "UObject/ObjectMacros.h"
@@ -104,6 +105,8 @@ FReply SImgMediaProcessImages::OnProcessImagesClicked()
 
 void SImgMediaProcessImages::ProcessAllImages(TSharedPtr<SNotificationItem> ConfirmNotification)
 {
+	TRACE_CPUPROFILER_EVENT_SCOPE(SImgMediaProcessImages::ProcessAllImages);
+
 	bool bUseCustomFormat = Options->bUseCustomFormat;
 	int32 InTileWidth = Options->TileSizeX;
 	int32 InTileHeight = Options->TileSizeY;
@@ -286,6 +289,7 @@ void SImgMediaProcessImages::ProcessImageCustom(TSharedPtr<IImageWrapper>& InIma
 	bool bHasAlphaChannel, const FString& InName)
 {
 #if IMGMEDIAEDITOR_EXR_SUPPORTED_PLATFORM
+	TRACE_CPUPROFILER_EVENT_SCOPE(SImgMediaProcessImages::ProcessImageCustom);
 	// Get image data.
 	ERGBFormat Format = InImageWrapper->GetFormat();
 	int32 Width = InImageWrapper->GetWidth();
@@ -294,7 +298,10 @@ void SImgMediaProcessImages::ProcessImageCustom(TSharedPtr<IImageWrapper>& InIma
 	int32 DestHeight = Height;
 	int32 BitDepth = InImageWrapper->GetBitDepth();
 	TArray64<uint8> RawData;
-	InImageWrapper->GetRaw(Format, BitDepth, RawData);
+	{
+		TRACE_CPUPROFILER_EVENT_SCOPE(SImgMediaProcessImages::ProcessImageCustom:GetRaw);
+		InImageWrapper->GetRaw(Format, BitDepth, RawData);
+	}
 
 	int32 NumTilesX = InTileWidth > 0 ? Width / InTileWidth : 1;
 	int32 NumTilesY = InTileHeight > 0 ? Height / InTileHeight : 1;
@@ -381,6 +388,7 @@ void SImgMediaProcessImages::ProcessImageCustom(TSharedPtr<IImageWrapper>& InIma
 	int32 CurrentMipBufferIndex = 0;
 
 	// Loop over each mip level.
+	TRACE_CPUPROFILER_EVENT_SCOPE(SImgMediaProcessImages::ProcessImageCustom:CreateMips);
 	int32 NumMips = OutFile.GetNumberOfMipLevels();
 	int32 MipSourceWidth = Width;
 	int32 MipSourceHeight = Height;
@@ -402,6 +410,7 @@ void SImgMediaProcessImages::ProcessImageCustom(TSharedPtr<IImageWrapper>& InIma
 		// Generate mip data.
 		if (MipLevel != 0)
 		{
+			TRACE_CPUPROFILER_EVENT_SCOPE(SImgMediaProcessImages::ProcessImageCustom:GenerateMipData);
 			int32 SourceStrideX = NumChannels;
 			int32 SourceStrideY = MipWidth * NumChannels * 2;
 			for (int32 PixelY = 0; PixelY < MipHeight; ++PixelY)
@@ -482,6 +491,7 @@ void SImgMediaProcessImages::ProcessImageCustom(TSharedPtr<IImageWrapper>& InIma
 		}
 
 		// Write to EXR.
+		TRACE_CPUPROFILER_EVENT_SCOPE(SImgMediaProcessImages::ProcessImageCustom:WriteEXR);
 		Stride.Y = MipWidth * BytesPerPixel;
 		int64 BufferOffset = 0;
 		int64 SingleBufferOffset = MipWidth * BytesPerPixelPerChannel;
@@ -515,6 +525,8 @@ void SImgMediaProcessImages::ProcessImageCustom(TSharedPtr<IImageWrapper>& InIma
 
 void SImgMediaProcessImages::RemoveAlphaChannel(TArray64<uint8>& Buffer)
 {
+	TRACE_CPUPROFILER_EVENT_SCOPE(SImgMediaProcessImages::RemoveAlphaChannel);
+
 	int32 BytesPerPixelPerChannel = 2;
 	int64 BufferSize = Buffer.Num() / BytesPerPixelPerChannel;
 	uint16* BufferPtr = (uint16*)(Buffer.GetData());
@@ -542,6 +554,8 @@ void SImgMediaProcessImages::TileData(uint8* SourceData, TArray64<uint8>& DestAr
 	int32 TileWidth, int32 TileHeight, int32 InTileBorder,
 	int32 BytesPerPixel)
 {
+	TRACE_CPUPROFILER_EVENT_SCOPE(SImgMediaProcessImages::TileData);
+
 	// We don't support tile borders larger than a tile size,
 	// but this shuld not happen in practice.
 	if ((InTileBorder > TileWidth) || (InTileBorder > TileHeight))
