ImgMedia: Added some profiling/instrumentation to the loader to track what it going on.

#rnx

diff --git a/Engine/Plugins/Media/ImgMedia/Source/ImgMedia/Private/ImgMediaModule.cpp b/Engine/Plugins/Media/ImgMedia/Source/ImgMedia/Private/ImgMediaModule.cpp
--- a/Engine/Plugins/Media/ImgMedia/Source/ImgMedia/Private/ImgMediaModule.cpp
+++ b/Engine/Plugins/Media/ImgMedia/Source/ImgMedia/Private/ImgMediaModule.cpp
@@ -13,7 +13,7 @@
 #include "ImgMediaSource.h"
 #include "IImgMediaModule.h"
 
-
+CSV_DEFINE_CATEGORY_MODULE(IMGMEDIA_API, ImgMedia, false);
 DEFINE_LOG_CATEGORY(LogImgMedia);
 
 FLazyName IImgMediaModule::CustomFormatAttributeName(TEXT("EpicGamesCustomFormat"));
diff --git a/Engine/Plugins/Media/ImgMedia/Source/ImgMedia/Private/Loader/ImgMediaLoader.cpp b/Engine/Plugins/Media/ImgMedia/Source/ImgMedia/Private/Loader/ImgMediaLoader.cpp
--- a/Engine/Plugins/Media/ImgMedia/Source/ImgMedia/Private/Loader/ImgMediaLoader.cpp
+++ b/Engine/Plugins/Media/ImgMedia/Source/ImgMedia/Private/Loader/ImgMediaLoader.cpp
@@ -17,6 +17,7 @@
 
 #include "GenericImgMediaReader.h"
 #include "IImageWrapperModule.h"
+#include "IImgMediaModule.h"
 #include "IImgMediaReader.h"
 #include "ImgMediaLoaderWork.h"
 #include "ImgMediaMipMapInfo.h"
@@ -545,6 +546,7 @@ IMediaSamples::EFetchBestSampleResult FImgMediaLoader::FetchBestVideoSampleForTi
 					if (Sample->Initialize(*Frame->Get(), SequenceDim, FMediaTimeStamp(FrameNumberToTime(MaxIdx), QueuedSampleFetch.CurrentSequenceIndex), FTimespan::FromSeconds(Duration), GetNumMipLevels(), ImgMediaLoader::GetTilingDescription(this)))
 					{
 						OutSample = Sample;
+						CSV_EVENT(ImgMedia, TEXT("LoaderFetchHit %d %d-%d"), MaxIdx, StartIndex, EndIndex);
 						return IMediaSamples::EFetchBestSampleResult::Ok;
 					}
 				}
@@ -578,6 +580,7 @@ IMediaSamples::EFetchBestSampleResult FImgMediaLoader::FetchBestVideoSampleForTi
 				}
 			}
 		}
+		CSV_EVENT(ImgMedia, TEXT("LoaderFetchMiss %d-%d"), StartIndex, EndIndex);
 	}
 	return IMediaSamples::EFetchBestSampleResult::NoSample;
 }
@@ -1313,6 +1316,8 @@ void FImgMediaLoader::Update(int32 PlayHeadFrame, float PlayRate, bool Loop)
 		}
 	}
 	Algo::Reverse(PendingFrameNumbers);
+
+	CSV_EVENT(ImgMedia, TEXT("LoaderUpdatePending %d %d"), (PendingFrameNumbers.Num() > 0) ? PendingFrameNumbers[0] : -1, (PendingFrameNumbers.Num() > 0) ? PendingFrameNumbers[PendingFrameNumbers.Num() - 1] : -1);
 }
 
 
diff --git a/Engine/Plugins/Media/ImgMedia/Source/ImgMedia/Private/Loader/ImgMediaLoaderWork.cpp b/Engine/Plugins/Media/ImgMedia/Source/ImgMedia/Private/Loader/ImgMediaLoaderWork.cpp
--- a/Engine/Plugins/Media/ImgMedia/Source/ImgMedia/Private/Loader/ImgMediaLoaderWork.cpp
+++ b/Engine/Plugins/Media/ImgMedia/Source/ImgMedia/Private/Loader/ImgMediaLoaderWork.cpp
@@ -4,7 +4,7 @@
 #include "ImgMediaPrivate.h"
 
 #include "Misc/ScopeLock.h"
-
+#include "IImgMediaModule.h"
 #include "IImgMediaReader.h"
 #include "ImgMediaLoader.h"
 
@@ -48,6 +48,7 @@ void FImgMediaLoaderWork::Initialize(int32 InFrameNumber,
 void FImgMediaLoaderWork::Abandon()
 {
 	SCOPE_CYCLE_COUNTER(STAT_ImgMedia_LoaderAbandonWork);
+	CSV_EVENT(ImgMedia, TEXT("LoaderAbandon %d"), FrameNumber);
 
 	Finalize(nullptr, 0.0f);
 }
@@ -66,6 +67,7 @@ void FImgMediaLoaderWork::DoThreadedWork()
 	else
 	{
 		SCOPE_CYCLE_COUNTER(STAT_ImgMedia_LoaderReadFrame);
+		TRACE_CPUPROFILER_EVENT_SCOPE_TEXT(*FString::Printf(TEXT("ImgMedia_LoaderReadFrame %d"), FrameNumber));
 
 		// read the image frame
 		if (ExistingFrame == nullptr)
diff --git a/Engine/Plugins/Media/ImgMedia/Source/ImgMedia/Public/IImgMediaModule.h b/Engine/Plugins/Media/ImgMedia/Source/ImgMedia/Public/IImgMediaModule.h
--- a/Engine/Plugins/Media/ImgMedia/Source/ImgMedia/Public/IImgMediaModule.h
+++ b/Engine/Plugins/Media/ImgMedia/Source/ImgMedia/Public/IImgMediaModule.h
@@ -4,6 +4,7 @@
 
 #include "CoreMinimal.h"
 #include "Delegates/Delegate.h"
+#include "ProfilingDebugging/CsvProfiler.h"
 #include "Templates/SharedPointer.h"
 #include "Modules/ModuleInterface.h"
 
@@ -12,6 +13,8 @@ class FImgMediaPlayer;
 class IMediaEventSink;
 class IMediaPlayer;
 
+CSV_DECLARE_CATEGORY_MODULE_EXTERN(IMGMEDIA_API, ImgMedia);
+
 /** Callback when a player gets created. */
 DECLARE_MULTICAST_DELEGATE_OneParam(FOnImgMediaPlayerCreated, const TSharedPtr<FImgMediaPlayer>&);
 
