ImgMedia: Added cancel button to process images.

#rnx

diff --git a/Engine/Plugins/Media/ImgMedia/Source/ImgMediaEditor/Private/Widgets/SImgMediaProcessImages.cpp b/Engine/Plugins/Media/ImgMedia/Source/ImgMediaEditor/Private/Widgets/SImgMediaProcessImages.cpp
--- a/Engine/Plugins/Media/ImgMedia/Source/ImgMediaEditor/Private/Widgets/SImgMediaProcessImages.cpp
+++ b/Engine/Plugins/Media/ImgMedia/Source/ImgMediaEditor/Private/Widgets/SImgMediaProcessImages.cpp
@@ -63,12 +63,30 @@ void SImgMediaProcessImages::Construct(const FArguments& InArgs)
 			.Padding(4.0f)
 			.HAlign(HAlign_Left)
 			[
-				SNew(SButton)
-					.OnClicked(this, &SImgMediaProcessImages::OnProcessImagesClicked)
-					.Text(LOCTEXT("StartProcessImages", "Process Images"))
-					.ToolTipText(LOCTEXT("ProcesssImagesButtonToolTip", "Start processing images."))
+				SNew(SHorizontalBox)
+
+				+ SHorizontalBox::Slot()
+					.AutoWidth()
+					[
+						SAssignNew(StartButton, SButton)
+							.OnClicked(this, &SImgMediaProcessImages::OnProcessImagesClicked)
+							.Text(LOCTEXT("StartProcessImages", "Process Images"))
+							.ToolTipText(LOCTEXT("ProcesssImagesButtonToolTip", "Start processing images."))
+					]
+
+				+ SHorizontalBox::Slot()
+					.AutoWidth()
+					[
+						SAssignNew(CancelButton, SButton)
+							.OnClicked(this, &SImgMediaProcessImages::OnCancelClicked)
+							.Text(LOCTEXT("CancelProcessImages", "Cancel"))
+							.ToolTipText(LOCTEXT("ProcesssImagesButtonToolTip", "Cancel processing images."))
+					]
 			]
 	];
+	bIsProcessing = false;
+	bIsCancelling = false;
+	UpdateWidgets();
 
 	// Create object with our options.
 	Options = TStrongObjectPtr<UImgMediaProcessImagesOptions>(NewObject<UImgMediaProcessImagesOptions>(GetTransientPackage(), NAME_None));
@@ -87,18 +105,44 @@ void SImgMediaProcessImages::Construct(const FArguments& InArgs)
 }
 END_SLATE_FUNCTION_BUILD_OPTIMIZATION
 
+void SImgMediaProcessImages::UpdateWidgets()
+{
+	StartButton->SetEnabled(!bIsProcessing);
+	CancelButton->SetEnabled(bIsProcessing && (!bIsCancelling));
+}
+
 FReply SImgMediaProcessImages::OnProcessImagesClicked()
 {
-	// Create notification.
-	FNotificationInfo Info(FText::GetEmpty());
-	Info.bFireAndForget = false;
-	TSharedPtr<SNotificationItem> ConfirmNotification = FSlateNotificationManager::Get().AddNotification(Info);
+	if (bIsProcessing == false)
+	{
+		// Set that we are processing now.
+		bIsProcessing = true;
+		UpdateWidgets();
+
+		// Create notification.
+		FNotificationInfo Info(FText::GetEmpty());
+		Info.bFireAndForget = false;
+
+		TSharedPtr<SNotificationItem> ConfirmNotification = FSlateNotificationManager::Get().AddNotification(Info);
+
+		
+		// Start async task to process files.
+		Async(EAsyncExecution::Thread, [this, ConfirmNotification]()
+		{
+			ProcessAllImages(ConfirmNotification);
+		});
+	}
+
+	return FReply::Handled();
+}
 
-	// Start async task to process files.
-	Async(EAsyncExecution::Thread, [this, ConfirmNotification]()
+FReply SImgMediaProcessImages::OnCancelClicked()
+{
+	if (bIsProcessing)
 	{
-		ProcessAllImages(ConfirmNotification);
-	});
+		bIsCancelling = true;
+		UpdateWidgets();
+	}
 
 	return FReply::Handled();
 }
@@ -191,19 +235,30 @@ void SImgMediaProcessImages::ProcessAllImages(TSharedPtr<SNotificationItem> Conf
 					Name = FPaths::ChangeExtension(Name, TEXT(""));
 					ProcessImage(ImageWrapper, InTileWidth, InTileHeight, Name, Ext);
 				}
+
+				// Do we want to cancel?
+				if (bIsCancelling)
+				{
+					break;
+				}
 			}
 		}
 	}
 
 	// Close notification. Must be run on the main thread.
-	Async(EAsyncExecution::TaskGraphMainThread, [ConfirmNotification]()
+	Async(EAsyncExecution::TaskGraphMainThread, [this, ConfirmNotification]()
 	{
 		if (ConfirmNotification.IsValid())
 		{
 			ConfirmNotification->SetEnabled(false);
-			ConfirmNotification->SetCompletionState(SNotificationItem::CS_Success);
+			ConfirmNotification->SetCompletionState(bIsCancelling ? SNotificationItem::CS_Fail : SNotificationItem::CS_Success);
 			ConfirmNotification->ExpireAndFadeout();
 		}
+
+		// Done with processing.
+		bIsProcessing = false;
+		bIsCancelling = false;
+		UpdateWidgets();
 	});
 }
 
diff --git a/Engine/Plugins/Media/ImgMedia/Source/ImgMediaEditor/Private/Widgets/SImgMediaProcessImages.h b/Engine/Plugins/Media/ImgMedia/Source/ImgMediaEditor/Private/Widgets/SImgMediaProcessImages.h
--- a/Engine/Plugins/Media/ImgMedia/Source/ImgMediaEditor/Private/Widgets/SImgMediaProcessImages.h
+++ b/Engine/Plugins/Media/ImgMedia/Source/ImgMediaEditor/Private/Widgets/SImgMediaProcessImages.h
@@ -10,6 +10,7 @@
 
 class IDetailsView;
 class IImageWrapper;
+class SButton;
 class SNotificationItem;
 
 /**
@@ -26,8 +27,15 @@ public:
 	void Construct(const FArguments& InArgs);
 
 private:
+	/**
+	 * Updates our widgets based on the current state.
+	 */
+	void UpdateWidgets();
+	
 	/** Called when we click on the process images button. */
 	FReply OnProcessImagesClicked();
+	/** Called when we click on the cancel button. */
+	FReply OnCancelClicked();
 
 	/**
 	 * Processes all images in the sequence and generate tiles/mips.
@@ -111,4 +119,13 @@ private:
 	TSharedPtr<class IDetailsView> DetailsView;
 	/** Object that holds our options. */
 	TStrongObjectPtr<UImgMediaProcessImagesOptions> Options;
+	/** Holds the button to start processing. */
+	TSharedPtr<SButton> StartButton;
+	/** Holds the button to cancel processing. */
+	TSharedPtr<SButton> CancelButton;
+
+	/** True if we are currently processing. */
+	bool bIsProcessing = false;
+	/** True if we want to cancel processing. */
+	bool bIsCancelling = false;
 };
