ImgMedia: ProcessImages can use a player to read in media.

#rnx

diff --git a/Engine/Plugins/Media/ImgMedia/Source/ImgMediaEditor/ImgMediaEditor.Build.cs b/Engine/Plugins/Media/ImgMedia/Source/ImgMediaEditor/ImgMediaEditor.Build.cs
--- a/Engine/Plugins/Media/ImgMedia/Source/ImgMediaEditor/ImgMediaEditor.Build.cs
+++ b/Engine/Plugins/Media/ImgMedia/Source/ImgMediaEditor/ImgMediaEditor.Build.cs
@@ -18,6 +18,7 @@ namespace UnrealBuildTool.Rules
 					"ImgMedia",
 					"InputCore",
 					"MediaAssets",
+					"RenderCore",
 					"Slate",
 					"SlateCore",
 					"UnrealEd",
diff --git a/Engine/Plugins/Media/ImgMedia/Source/ImgMediaEditor/Private/Widgets/SImgMediaProcessImages.cpp b/Engine/Plugins/Media/ImgMedia/Source/ImgMediaEditor/Private/Widgets/SImgMediaProcessImages.cpp
--- a/Engine/Plugins/Media/ImgMedia/Source/ImgMediaEditor/Private/Widgets/SImgMediaProcessImages.cpp
+++ b/Engine/Plugins/Media/ImgMedia/Source/ImgMediaEditor/Private/Widgets/SImgMediaProcessImages.cpp
@@ -4,14 +4,21 @@
 
 #include "Async/Async.h"
 #include "Customizations/ImgMediaFilePathCustomization.h"
+#include "Editor.h"
 #include "EditorStyleSet.h"
+#include "Engine/Canvas.h"
+#include "Engine/TextureRenderTarget2D.h"
 #include "Framework/Notifications/NotificationManager.h"
 #include "HAL/PlatformFileManager.h"
 #include "IImageWrapperModule.h"
 #include "IImgMediaModule.h"
+#include "ImageUtils.h"
 #include "ImageWrapperHelper.h"
 #include "ImgMediaEditorModule.h"
 #include "ImgMediaProcessImagesOptions.h"
+#include "MediaPlayer.h"
+#include "MediaSource.h"
+#include "MediaTexture.h"
 #include "Misc/FileHelper.h"
 #include "Misc/Paths.h"
 #include "Modules/ModuleManager.h"
@@ -32,8 +39,11 @@
 
 #define LOCTEXT_NAMESPACE "ImgMediaProcessImages"
 
+#define IMGMEDIA_PROCESSIMAGES_USE_PLAYER 0
+
 SImgMediaProcessImages::~SImgMediaProcessImages()
 {
+	CleanUp();
 }
 
 BEGIN_SLATE_FUNCTION_BUILD_OPTIMIZATION
@@ -105,6 +115,14 @@ void SImgMediaProcessImages::Construct(const FArguments& InArgs)
 }
 END_SLATE_FUNCTION_BUILD_OPTIMIZATION
 
+void SImgMediaProcessImages::Tick(const FGeometry& AllottedGeometry, const double InCurrentTime, const float InDeltaTime)
+{
+	SCompoundWidget::Tick(AllottedGeometry, InCurrentTime, InDeltaTime);
+#if IMGMEDIA_PROCESSIMAGES_USE_PLAYER
+	HandleProcessing();
+#endif
+}
+
 void SImgMediaProcessImages::UpdateWidgets()
 {
 	StartButton->SetEnabled(!bIsProcessing);
@@ -123,14 +141,43 @@ FReply SImgMediaProcessImages::OnProcessImagesClicked()
 		FNotificationInfo Info(FText::GetEmpty());
 		Info.bFireAndForget = false;
 
-		TSharedPtr<SNotificationItem> ConfirmNotification = FSlateNotificationManager::Get().AddNotification(Info);
+		ConfirmNotification = FSlateNotificationManager::Get().AddNotification(Info);
+
+#if IMGMEDIA_PROCESSIMAGES_USE_PLAYER
+		// Create player.
+		MediaPlayer = NewObject<UMediaPlayer>(GetTransientPackage(), "MediaPlayer", RF_Transient);
+		MediaPlayer->SetLooping(false);
+		MediaPlayer->PlayOnOpen = true;
+		MediaPlayer->AddToRoot();
+
+		// Create texture.
+		MediaTexture = NewObject<UMediaTexture>(GetTransientPackage(), "MediaTexture", RF_Transient);
+		MediaTexture->SetMediaPlayer(MediaPlayer);
+		MediaTexture->UpdateResource();
+		MediaTexture->AddToRoot();
 
+		// Create media source.
+		MediaSource = UMediaSource::SpawnMediaSourceForString(Options->SequencePath.FilePath);
+		if (MediaSource == nullptr)
+		{
+			return FReply::Handled();
+		}
+		MediaSource->AddToRoot();
+
+		// Start playing.
+		CurrentFrameIndex = 0;
+		CurrentTime = FTimespan::FromSeconds(0.0f);
+		MediaPlayer->SetBlockOnTimeRange(TRange<FTimespan>(CurrentTime,
+			CurrentTime + FTimespan::FromSeconds(1.0f /100000.0f)));
+		MediaPlayer->OpenSource(MediaSource);
+#else
 		
 		// Start async task to process files.
-		Async(EAsyncExecution::Thread, [this, ConfirmNotification]()
+		Async(EAsyncExecution::Thread, [this]()
 		{
-			ProcessAllImages(ConfirmNotification);
+			ProcessAllImages();
 		});
+#endif
 	}
 
 	return FReply::Handled();
@@ -147,7 +194,7 @@ FReply SImgMediaProcessImages::OnCancelClicked()
 	return FReply::Handled();
 }
 
-void SImgMediaProcessImages::ProcessAllImages(TSharedPtr<SNotificationItem> ConfirmNotification)
+void SImgMediaProcessImages::ProcessAllImages()
 {
 	TRACE_CPUPROFILER_EVENT_SCOPE(SImgMediaProcessImages::ProcessAllImages);
 
@@ -194,14 +241,15 @@ void SImgMediaProcessImages::ProcessAllImages(TSharedPtr<SNotificationItem> Conf
 			// Loop through all files.
 			int NumDone = 0;
 			int TotalNum = FoundFiles.Num();
+			TSharedPtr<SNotificationItem> LocalConfirmNotification = ConfirmNotification;
 			for (const FString& FileName : FoundFiles)
 			{
 				// Update notification with current status.
-				Async(EAsyncExecution::TaskGraphMainThread, [ConfirmNotification, NumDone, TotalNum]()
+				Async(EAsyncExecution::TaskGraphMainThread, [LocalConfirmNotification, NumDone, TotalNum]()
 				{
-					if (ConfirmNotification.IsValid())
+					if (LocalConfirmNotification.IsValid())
 					{
-						ConfirmNotification->SetText(
+						LocalConfirmNotification->SetText(
 							FText::Format(LOCTEXT("ImgMediaCompleted", "ImgMedia Completed {0}/{1}"),
 								FText::AsNumber(NumDone), FText::AsNumber(TotalNum)));
 					}
@@ -246,7 +294,7 @@ void SImgMediaProcessImages::ProcessAllImages(TSharedPtr<SNotificationItem> Conf
 	}
 
 	// Close notification. Must be run on the main thread.
-	Async(EAsyncExecution::TaskGraphMainThread, [this, ConfirmNotification]()
+	Async(EAsyncExecution::TaskGraphMainThread, [this]()
 	{
 		if (ConfirmNotification.IsValid())
 		{
@@ -349,15 +397,29 @@ void SImgMediaProcessImages::ProcessImageCustom(TSharedPtr<IImageWrapper>& InIma
 	ERGBFormat Format = InImageWrapper->GetFormat();
 	int32 Width = InImageWrapper->GetWidth();
 	int32 Height = InImageWrapper->GetHeight();
-	int32 DestWidth = Width;
-	int32 DestHeight = Height;
 	int32 BitDepth = InImageWrapper->GetBitDepth();
 	TArray64<uint8> RawData;
 	{
 		TRACE_CPUPROFILER_EVENT_SCOPE(SImgMediaProcessImages::ProcessImageCustom:GetRaw);
 		InImageWrapper->GetRaw(Format, BitDepth, RawData);
 	}
+	ProcessImageCustomRawData(RawData, Width, Height, BitDepth,
+		InTileWidth, InTileHeight, InTileBorder, bInEnableMips,
+		bHasAlphaChannel, InName);
+#else // IMGMEDIAEDITOR_EXR_SUPPORTED_PLATFORM
+	UE_LOG(LogImgMediaEditor, Error, TEXT("EXR not supported on this platform."));
+#endif // IMGMEDIAEDITOR_EXR_SUPPORTED_PLATFORM
+}
 
+void SImgMediaProcessImages::ProcessImageCustomRawData(TArray64<uint8>& RawData,
+	int32 Width, int32 Height, int32 BitDepth,
+	int32 InTileWidth, int32 InTileHeight, int32 InTileBorder, bool bInEnableMips,
+	bool bHasAlphaChannel, const FString& InName)
+{
+#if IMGMEDIAEDITOR_EXR_SUPPORTED_PLATFORM
+	TRACE_CPUPROFILER_EVENT_SCOPE(SImgMediaProcessImages::ProcessImageCustomRawData);
+	int32 DestWidth = Width;
+	int32 DestHeight = Height;
 	int32 NumTilesX = InTileWidth > 0 ? Width / InTileWidth : 1;
 	int32 NumTilesY = InTileHeight > 0 ? Height / InTileHeight : 1;
 	int32 TileWidth = Width / NumTilesX;
@@ -693,4 +755,189 @@ void SImgMediaProcessImages::TileData(uint8* SourceData, TArray64<uint8>& DestAr
 	}
 }
 
+void SImgMediaProcessImages::HandleProcessing()
+{
+	// Are we processing?
+	if (bIsProcessing)
+	{
+		// We did not cancel yet?
+		bool bShouldExit = false;
+		if ((MediaPlayer != nullptr) && (bIsCancelling == false))
+		{
+			UE_LOG(LogImgMediaEditor, Verbose,
+				TEXT("ProcessImages Time:%f PlayerTime:%f Duration:%f"),
+				CurrentTime.GetTotalSeconds(),
+				MediaPlayer->GetTime().GetTotalSeconds(),
+				MediaPlayer->GetDuration().GetTotalSeconds());
+
+			// Has the player stopped playing?
+			if (MediaPlayer->IsClosed())
+			{
+				bShouldExit = true;
+			}
+			// Are we playing?
+			else if (MediaPlayer->IsPlaying())
+			{
+				// Are we set up yet?
+				if (RenderTarget == nullptr)
+				{
+					CreateRenderTarget();
+					// Get frame duration.
+					float FrameRate = MediaPlayer->GetVideoTrackFrameRate(INDEX_NONE, INDEX_NONE);
+					if (FrameRate <= 0.0f)
+					{
+						FrameRate = 24.0f;
+					}
+					FrameDuration = FTimespan::FromSeconds(1.0f / FrameRate);
+				}
+
+				// Copy media texture to our render target.
+				DrawTextureToRenderTarget();
+
+				// Process this render.
+				TArray64<uint8> RawData;
+				bool bReadSuccess = FImageUtils::GetRawData(RenderTarget, RawData);
+				if (bReadSuccess)
+				{
+					int32 Width = RenderTarget->GetSurfaceWidth();
+					int32 Height = RenderTarget->GetSurfaceHeight();
+					int32 BitDepth = 16;
+					bool bUseCustomFormat = Options->bUseCustomFormat;
+					int32 InTileWidth = Options->TileSizeX;
+					int32 InTileHeight = Options->TileSizeY;
+					int32 TileBorder = Options->TileBorder;
+					bool bEnableMips = Options->bEnableMipMapping;
+					bool bHasAlphaChannel = false;
+					FString OutPath = Options->OutputPath.Path;
+					FString FileName = FString::Format(TEXT("image{0}.exr"),
+						{CurrentFrameIndex});
+					FString Name = FPaths::Combine(OutPath, FileName);
+
+					ProcessImageCustomRawData(RawData, Width, Height, BitDepth,
+						InTileWidth, InTileHeight, TileBorder, bEnableMips,
+						bHasAlphaChannel, Name);
+				}
+				else
+				{
+					UE_LOG(LogImgMediaEditor, Error, TEXT("ProcessImages failed to get raw data."));
+				}
+
+				// Update notification.
+				if (ConfirmNotification.IsValid())
+				{
+					ConfirmNotification->SetText(
+						FText::Format(LOCTEXT("ImgMediaCompleted", "ImgMedia Completed {0}"),
+							FText::AsNumber(CurrentFrameIndex)));
+				}
+
+				// Next frame.
+				CurrentTime += FrameDuration;
+				CurrentFrameIndex++;
+				if (CurrentTime >= MediaPlayer->GetDuration())
+				{
+					bShouldExit = true;
+				}
+				else
+				{
+					MediaPlayer->SetBlockOnTimeRange(TRange<FTimespan>(CurrentTime, CurrentTime + FrameDuration));
+				}
+			}
+		}
+		else
+		{
+			bShouldExit = true;
+		}
+
+		// Are we done?
+		if (bShouldExit)
+		{
+			// Remove notification.
+			if (ConfirmNotification.IsValid())
+			{
+				ConfirmNotification->SetEnabled(false);
+				ConfirmNotification->SetCompletionState(bIsCancelling ? SNotificationItem::CS_Fail : SNotificationItem::CS_Success);
+				ConfirmNotification->ExpireAndFadeout();
+				ConfirmNotification.Reset();
+			}
+
+			// Clean up.
+			bIsCancelling = false;
+			bIsProcessing = false;
+			CleanUp();
+			UpdateWidgets();
+		}
+	}
+}
+
+void SImgMediaProcessImages::CreateRenderTarget()
+{
+	if (MediaTexture != nullptr)
+	{
+		int32 Width = MediaTexture->GetWidth();
+		int32 Height = MediaTexture->GetHeight();
+
+		RenderTarget = NewObject<UTextureRenderTarget2D>(GetTransientPackage(), TEXT("ImgMediaProcessImages"));
+		RenderTarget->RenderTargetFormat = RTF_RGBA16f;
+		RenderTarget->InitAutoFormat(Width, Height);
+		RenderTarget->AddToRoot();
+		RenderTarget->UpdateResourceImmediate(true);
+	}
+}
+
+void SImgMediaProcessImages::DrawTextureToRenderTarget()
+{
+	UWorld* World = nullptr;
+	World = GEditor->GetEditorWorldContext().World();
+	World->FlushDeferredParameterCollectionInstanceUpdates();
+
+	FTextureRenderTargetResource* RenderTargetResource = RenderTarget->GameThread_GetRenderTargetResource();
+
+	UCanvas* Canvas = World->GetCanvasForDrawMaterialToRenderTarget();
+	FCanvas RenderCanvas(
+		RenderTargetResource,
+		nullptr,
+		World,
+		World->FeatureLevel);
+	Canvas->Init(RenderTarget->SizeX, RenderTarget->SizeY, nullptr, &RenderCanvas);
+	Canvas->Update();
+
+	{
+		ENQUEUE_RENDER_COMMAND(FlushDeferredResourceUpdateCommand)(
+			[RenderTargetResource](FRHICommandListImmediate& RHICmdList)
+		{
+			RenderTargetResource->FlushDeferredResourceUpdate(RHICmdList);
+		});
+
+		Canvas->K2_DrawTexture(MediaTexture, FVector2D(0, 0), FVector2D(RenderTarget->SizeX, RenderTarget->SizeY), FVector2D(0, 0));
+
+		RenderCanvas.Flush_GameThread();
+		Canvas->Canvas = nullptr;
+		RenderTarget->UpdateResourceImmediate(false);
+	}
+}
+
+void SImgMediaProcessImages::CleanUp()
+{
+	if (MediaPlayer != nullptr)
+	{
+		MediaPlayer->RemoveFromRoot();
+		MediaPlayer = nullptr;
+	}
+	if (MediaTexture != nullptr)
+	{
+		MediaTexture->RemoveFromRoot();
+		MediaTexture = nullptr;
+	}
+	if (MediaSource != nullptr)
+	{
+		MediaSource->RemoveFromRoot();
+		MediaSource = nullptr;
+	}
+	if (RenderTarget != nullptr)
+	{
+		RenderTarget->RemoveFromRoot();
+		RenderTarget = nullptr;
+	}
+}
+
 #undef LOCTEXT_NAMESPACE
diff --git a/Engine/Plugins/Media/ImgMedia/Source/ImgMediaEditor/Private/Widgets/SImgMediaProcessImages.h b/Engine/Plugins/Media/ImgMedia/Source/ImgMediaEditor/Private/Widgets/SImgMediaProcessImages.h
--- a/Engine/Plugins/Media/ImgMedia/Source/ImgMediaEditor/Private/Widgets/SImgMediaProcessImages.h
+++ b/Engine/Plugins/Media/ImgMedia/Source/ImgMediaEditor/Private/Widgets/SImgMediaProcessImages.h
@@ -12,6 +12,10 @@ class IDetailsView;
 class IImageWrapper;
 class SButton;
 class SNotificationItem;
+class UMediaPlayer;
+class UMediaSource;
+class UMediaTexture;
+class UTextureRenderTarget2D;
 
 /**
  * SImgMediaProcessImages provides processing of image sequences.
@@ -26,6 +30,10 @@ public:
 
 	void Construct(const FArguments& InArgs);
 
+	/** SWidget interface */
+	virtual void Tick(const FGeometry& AllottedGeometry, const double InCurrentTime, const float InDeltaTime);
+
+
 private:
 	/**
 	 * Updates our widgets based on the current state.
@@ -39,10 +47,8 @@ private:
 
 	/**
 	 * Processes all images in the sequence and generate tiles/mips.
-	 *
-	 * @param ConfirmNotification	Notification that will be updated with progress and closed when done.
 	 */
-	void ProcessAllImages(TSharedPtr<SNotificationItem> ConfirmNotification);
+	void ProcessAllImages();
 	
 	/**
 	  * Checks if this file has an alpha channel.
@@ -84,6 +90,27 @@ private:
 		int32 InTileWidth, int32 InTileHeight, int32 InTileBorder, 
 		bool bInEnableMips, bool bHasAlphaChannel, const FString& InName);
 
+	/**
+	 * Processess a single image and writes out a file.
+	 * Tiles and mips may be generated.
+	 * This might not run on the game thread.
+	 *
+	 * @param RawData			Image pixel data (RGBA FFloat16).
+	 * @param Width				Width of image.
+	 * @param Height			Height of image.
+	 * @param BitDepth			Number of bits per pixel.
+	 * @param InTileWidth		Desired width of tiles.
+	 * @param InTileHeight		Desired height of tiles.
+	 * @param InTileBorder		Number of pixels to duplicate along a tile edge.
+	 * @param bInEnableMips		Turn on mip mapping.
+	 * @param bHasAlphaChannel	True if there really is an alpha channel.
+	 * @param InName			Full path and name of file to write.
+	 */
+	void ProcessImageCustomRawData(TArray64<uint8>& RawData,
+		int32 Width, int32 Height, int32 BitDepth,
+		int32 InTileWidth, int32 InTileHeight, int32 InTileBorder, bool bInEnableMips,
+		bool bHasAlphaChannel, const FString& InName);
+	
 	/**
 	 * Removes the alpha channel from a buffer.
 	 * This will be done in place.
@@ -115,6 +142,26 @@ private:
 		int32 TileWidth, int32 TileHeight, int32 InTileBorder,
 		int32 BytesPerPixel);
 
+	/**
+	 * Call this every frame to handle processing.
+	 */
+	void HandleProcessing();
+
+	/**
+	 * Creates a render target.
+	 */
+	void CreateRenderTarget();
+
+	/**
+	 * Draws the media texture to our render target.
+	 */
+	void DrawTextureToRenderTarget();
+
+	/**
+	 * Clean up after processing is done.
+	 */
+	void CleanUp();
+
 	/** Holds our details view. */
 	TSharedPtr<class IDetailsView> DetailsView;
 	/** Object that holds our options. */
@@ -123,9 +170,27 @@ private:
 	TSharedPtr<SButton> StartButton;
 	/** Holds the button to cancel processing. */
 	TSharedPtr<SButton> CancelButton;
+	/** Notification to update user on our progress. */
+	TSharedPtr<SNotificationItem> ConfirmNotification;
+
+	/** Used to read in media. */
+	TObjectPtr<UMediaPlayer> MediaPlayer;
+	/** Texture for the player. */
+	TObjectPtr<UMediaTexture> MediaTexture;
+	/** Source for the player. */
+	TObjectPtr<UMediaSource> MediaSource;
+	/** Render target to render the media texture to. */
+	TObjectPtr<UTextureRenderTarget2D> RenderTarget;
+	/** The current time we are processing. */
+	FTimespan CurrentTime;
+	/** How long a frame is. */
+	FTimespan FrameDuration;
+	/** Current frame index we are processinig. */
+	int32 CurrentFrameIndex = 0;
 
 	/** True if we are currently processing. */
 	bool bIsProcessing = false;
 	/** True if we want to cancel processing. */
 	bool bIsCancelling = false;
+	
 };
