ImgMedia: Fixed process images not writing out the last frame.
#rnx

diff --git a/Engine/Plugins/Media/ImgMedia/Source/ImgMediaEditor/Private/Widgets/SImgMediaProcessImages.cpp b/Engine/Plugins/Media/ImgMedia/Source/ImgMediaEditor/Private/Widgets/SImgMediaProcessImages.cpp
--- a/Engine/Plugins/Media/ImgMedia/Source/ImgMediaEditor/Private/Widgets/SImgMediaProcessImages.cpp
+++ b/Engine/Plugins/Media/ImgMedia/Source/ImgMediaEditor/Private/Widgets/SImgMediaProcessImages.cpp
@@ -16,6 +16,7 @@
 #include "ImageWrapperHelper.h"
 #include "ImgMediaEditorModule.h"
 #include "ImgMediaProcessImagesOptions.h"
+#include "Math/UnrealMathUtility.h"
 #include "MediaPlayer.h"
 #include "MediaSource.h"
 #include "MediaTexture.h"
@@ -146,7 +147,7 @@ FReply SImgMediaProcessImages::OnProcessImagesClicked()
 #if IMGMEDIA_PROCESSIMAGES_USE_PLAYER
 		// Create player.
 		MediaPlayer = NewObject<UMediaPlayer>(GetTransientPackage(), "MediaPlayer", RF_Transient);
-		MediaPlayer->SetLooping(false);
+		MediaPlayer->SetLooping(true);
 		MediaPlayer->PlayOnOpen = true;
 		MediaPlayer->AddToRoot();
 
@@ -764,19 +765,26 @@ void SImgMediaProcessImages::HandleProcessing()
 		bool bShouldExit = false;
 		if ((MediaPlayer != nullptr) && (bIsCancelling == false))
 		{
+			// Get which frame the player is on.
+			int32 PlayerFrame = 0;
+			if (FrameDuration.GetTotalSeconds() > 0.0f)
+			{
+				PlayerFrame = (int32)FMath::RoundHalfToEven(MediaPlayer->GetTime().GetTotalSeconds() / FrameDuration.GetTotalSeconds());
+			}
+
 			UE_LOG(LogImgMediaEditor, Verbose,
-				TEXT("ProcessImages Time:%f PlayerTime:%f Duration:%f"),
+				TEXT("ProcessImages Time:%f PlayerTime:%f Duration:%f Frame:%d"),
 				CurrentTime.GetTotalSeconds(),
 				MediaPlayer->GetTime().GetTotalSeconds(),
-				MediaPlayer->GetDuration().GetTotalSeconds());
+				MediaPlayer->GetDuration().GetTotalSeconds(), PlayerFrame);
 
 			// Has the player stopped playing?
 			if (MediaPlayer->IsClosed())
 			{
 				bShouldExit = true;
 			}
-			// Are we playing?
-			else if (MediaPlayer->IsPlaying())
+			// Is this the frame we want?
+			else if ((MediaPlayer->IsPreparing() == false) && (CurrentFrameIndex == PlayerFrame))
 			{
 				// Are we set up yet?
 				if (RenderTarget == nullptr)
@@ -920,6 +928,7 @@ void SImgMediaProcessImages::CleanUp()
 {
 	if (MediaPlayer != nullptr)
 	{
+		MediaPlayer->Close();
 		MediaPlayer->RemoveFromRoot();
 		MediaPlayer = nullptr;
 	}
