ImgMediaPlayer ExrReaderGPU: Added same file mip reading capability for Vanilla EXR.

diff --git a/Engine/Plugins/Media/ImgMedia/Source/ExrReaderGpu/Private/ExrReaderGpu.cpp b/Engine/Plugins/Media/ImgMedia/Source/ExrReaderGpu/Private/ExrReaderGpu.cpp
--- a/Engine/Plugins/Media/ImgMedia/Source/ExrReaderGpu/Private/ExrReaderGpu.cpp
+++ b/Engine/Plugins/Media/ImgMedia/Source/ExrReaderGpu/Private/ExrReaderGpu.cpp
@@ -251,7 +251,7 @@ bool FExrReader::GenerateTextureData(uint16* Buffer, int32 BufferSize, FString F
 	return true;
 }
 
-bool FExrReader::OpenExrAndPrepareForPixelReading(FString FilePath, int32 NumOffsets)
+bool FExrReader::OpenExrAndPrepareForPixelReading(FString FilePath, int32 NumOffsets, int32 NumLevels)
 {
 	if (FileHandle != nullptr)
 	{
@@ -273,14 +273,18 @@ bool FExrReader::OpenExrAndPrepareForPixelReading(FString FilePath, int32 NumOff
 		{
 			return false;
 		}
-		LineOrTileOffsets.SetNum(NumOffsets);
+		LineOrTileOffsetsPerLevel.SetNum(NumLevels);
 
 		// At the moment we support only increasing Y.
 		ELineOrder LineOrder = INCREASING_Y;
-		ReadLineOrTileOffsets(FileHandle, LineOrder, LineOrTileOffsets);
+		for (int Level = 0; Level < NumLevels; Level++)
+		{
+			LineOrTileOffsetsPerLevel[Level].SetNum(NumOffsets >> (Level * 2));
+			ReadLineOrTileOffsets(FileHandle, LineOrder, LineOrTileOffsetsPerLevel[Level]);
+		}
 	}
 
-	fseek(FileHandle, LineOrTileOffsets[0], SEEK_SET);
+	fseek(FileHandle, LineOrTileOffsetsPerLevel[0][0], SEEK_SET);
 
 	return true;
 }
@@ -306,24 +310,24 @@ bool FExrReader::ReadExrImageChunk(void* Buffer, int64 ChunkSize)
 	return NumElementsRead == NumElements;
 }
 
-bool FExrReader::SeekTileWithinFile(const int32 StartTileIndex, const FIntPoint& TexDimInTiles, int64& OutBufferOffset)
+bool FExrReader::SeekTileWithinFile(const int32 StartTileIndex, const FIntPoint& TexDimInTiles, const int32 Level, int64& OutBufferOffset)
 {
-	if (StartTileIndex >= LineOrTileOffsets.Num())
+	if (StartTileIndex >= LineOrTileOffsetsPerLevel[Level].Num())
 	{
 		UE_LOG(LogExrReaderGpu, Error, TEXT("Tile index is invalid."));
 		return false;
 	}
-	int64 LineOffset = LineOrTileOffsets[StartTileIndex];
-	OutBufferOffset = LineOffset - LineOrTileOffsets[0];
+	int64 LineOffset = LineOrTileOffsetsPerLevel[Level][StartTileIndex];
+	OutBufferOffset = LineOffset - LineOrTileOffsetsPerLevel[Level][0];
 	return fseek(FileHandle, LineOffset, SEEK_SET) == 0;
 }
 
-bool FExrReader::SeekTileWithinFileCustom(const int32 StartTileIndex, const int64 TileStride, int64& OutBufferOffset)
+bool FExrReader::SeekTileWithinFileCustom(const int32 StartTileIndex, const int64 TileStride, const int32 Level, int64& OutBufferOffset)
 {
 	// Our custom exr is written as one single tile. In Exr every tile starts with 20 byte padding.
 	// Therefore we want to ignore it, when we read custom tiles.
 	const int32 TilePadding = 20;
-	int64 LineOffset = TilePadding + LineOrTileOffsets[0] + StartTileIndex * TileStride;
+	int64 LineOffset = TilePadding + LineOrTileOffsetsPerLevel[Level][0] + StartTileIndex * TileStride;
 	OutBufferOffset = StartTileIndex * TileStride;
 	return fseek(FileHandle, LineOffset, SEEK_SET) == 0;
 }
@@ -336,7 +340,7 @@ bool FExrReader::CloseExrFile()
 		return false;
 	}
 	fclose(FileHandle);
-	LineOrTileOffsets.Empty();
+	LineOrTileOffsetsPerLevel.Empty();
 	return true;
 }
 
diff --git a/Engine/Plugins/Media/ImgMedia/Source/ExrReaderGpu/Public/ExrReaderGpu.h b/Engine/Plugins/Media/ImgMedia/Source/ExrReaderGpu/Public/ExrReaderGpu.h
--- a/Engine/Plugins/Media/ImgMedia/Source/ExrReaderGpu/Public/ExrReaderGpu.h
+++ b/Engine/Plugins/Media/ImgMedia/Source/ExrReaderGpu/Public/ExrReaderGpu.h
@@ -64,7 +64,7 @@ public:
 	* Reading in chunks allows the process of reading to be canceled midway through reading. 
 	* This function reads and discards the header and keeps the pointers to the scanlines. 
 	*/
-	bool OpenExrAndPrepareForPixelReading(FString FilePath, int32 NumOffsets);
+	bool OpenExrAndPrepareForPixelReading(FString FilePath, int32 NumOffsets, int32 NumLevels);
 
 	/** 
 	* Read a chunk of an Exr file previously open via OpenExrAndPrepareForPixelReading.
@@ -72,15 +72,15 @@ public:
 	*/
 	bool ReadExrImageChunk(void* Buffer, int64 ChunkSize);
 
-	/** 
+	/**
 	* Moves to the specified tile position in the file in preparation for reading.
 	*/
-	bool SeekTileWithinFile(const int32 StartTileIndex, const FIntPoint& TexDimInTiles, int64& OutBufferOffset);
+	bool SeekTileWithinFile(const int32 StartTileIndex, const FIntPoint& TexDimInTiles, const int32 Level, int64& OutBufferOffset);
 
 	/** 
 	* Moves to the specified tile position in the file in preparation for reading.
 	*/
-	bool SeekTileWithinFileCustom(const int32 StartTileIndex, const int64 TileStride, int64& OutBufferOffset);
+	bool SeekTileWithinFileCustom(const int32 StartTileIndex, const int64 TileStride, const int32 Level, int64& OutBufferOffset);
 
 	/**
 	* After we have finished reading the file that was open via OpenExrAndPrepareForPixelReading
@@ -99,6 +99,6 @@ private:
 	/**
 	* These are byte offsets pointing to scanlines or tiles from the begining of the file.
 	*/
-	TArray<int64> LineOrTileOffsets;
+	TArray<TArray<int64>> LineOrTileOffsetsPerLevel;
 };
 #endif
diff --git a/Engine/Plugins/Media/ImgMedia/Source/ImgMedia/Private/Loader/ImgMediaLoader.cpp b/Engine/Plugins/Media/ImgMedia/Source/ImgMedia/Private/Loader/ImgMediaLoader.cpp
--- a/Engine/Plugins/Media/ImgMedia/Source/ImgMedia/Private/Loader/ImgMediaLoader.cpp
+++ b/Engine/Plugins/Media/ImgMedia/Source/ImgMedia/Private/Loader/ImgMediaLoader.cpp
@@ -861,6 +861,8 @@ void FImgMediaLoader::LoadSequence(const FString& SequencePath, const FFrameRate
 		return;
 	}
 
+	NumMipLevels = FMath::Max(FirstFrameInfo.NumMipLevels, ImagePaths.Num());
+
 	SequenceDim = FirstFrameInfo.Dim;
 
 	if (FrameRateOverride.IsValid() && (FrameRateOverride.Numerator > 0))
diff --git a/Engine/Plugins/Media/ImgMedia/Source/ImgMedia/Private/Loader/ImgMediaLoader.h b/Engine/Plugins/Media/ImgMedia/Source/ImgMedia/Private/Loader/ImgMediaLoader.h
--- a/Engine/Plugins/Media/ImgMedia/Source/ImgMedia/Private/Loader/ImgMediaLoader.h
+++ b/Engine/Plugins/Media/ImgMedia/Source/ImgMedia/Private/Loader/ImgMediaLoader.h
@@ -213,10 +213,18 @@ public:
 	/**
 	 * Get the number of mipmap levels we have.
 	 */
-	int32 GetNumMipLevels() const
+	int32 GetNumMipLevels() const 
+	{ 
+		return NumMipLevels; 
+	};
+
+	/**
+	 * If mips are stored in separate files we should make sure to read it.
+	 */
+	bool MipsInSeparateFiles() const
 	{
-		return ImagePaths.Num();
-	}
+		return ImagePaths.Num() > 1;
+	};
 
 	/**
 	 * Get the number of images in a single mip level.
@@ -503,6 +511,9 @@ private:
 	/** Total number of frames to load. */
 	int32 NumFramesToLoad;
 
+	/** Stores num of mip levels if the sequence is made out of files that contain all mips in one file. */
+	int32 NumMipLevels;
+
 	/** The image sequence reader to use. */
 	TSharedPtr<IImgMediaReader, ESPMode::ThreadSafe> Reader;
 
diff --git a/Engine/Plugins/Media/ImgMedia/Source/ImgMedia/Private/Readers/ExrImgMediaReader.cpp b/Engine/Plugins/Media/ImgMedia/Source/ImgMedia/Private/Readers/ExrImgMediaReader.cpp
--- a/Engine/Plugins/Media/ImgMedia/Source/ImgMedia/Private/Readers/ExrImgMediaReader.cpp
+++ b/Engine/Plugins/Media/ImgMedia/Source/ImgMedia/Private/Readers/ExrImgMediaReader.cpp
@@ -72,22 +72,23 @@ bool FExrImgMediaReader::GetFrameInfo(const FString& ImagePath, FImgMediaFrameIn
 }
 
 FExrImgMediaReader::EReadResult FExrImgMediaReader::ReadTilesCustom
-(uint16* Buffer
+	( uint16* Buffer
 	, const FString& ImagePath
 	, int32 FrameId
 	, const FIntRect& TileRegion
-	, const FIntPoint& FullTexDimInTiles
-	, const FIntPoint& TileDim
-	, const int32 PixelSize
-	, const bool bCustomExr)
+	, TSharedPtr<FSampleConverterParameters> ConverterParams
+	, const int32 CurrentMipLevel)
 {
 #if defined(PLATFORM_WINDOWS) && PLATFORM_WINDOWS
 	EReadResult bResult = Success;
 
 	int64 CurrentBufferPos = 0;
 	FExrReader ChunkReader;
-	int32 NumTiles = bCustomExr ? 1 : FullTexDimInTiles.X * FullTexDimInTiles.Y;
-	if (!ChunkReader.OpenExrAndPrepareForPixelReading(ImagePath, NumTiles))
+	int MipLevelDiv = 1 << CurrentMipLevel;
+	FIntPoint DimensionInTiles = ConverterParams->FrameInfo.NumTiles / MipLevelDiv;
+	int32 NumTiles = ConverterParams->bCustomExr ? 1 : ConverterParams->FrameInfo.NumTiles.X * ConverterParams->FrameInfo.NumTiles.Y;
+	int32 NumMipLevels = ConverterParams->bMipsInSeparateFiles ? 1 : ConverterParams->NumMipLevels;
+	if (!ChunkReader.OpenExrAndPrepareForPixelReading(ImagePath, NumTiles, NumMipLevels))
 	{
 		return Fail;
 	}
@@ -104,20 +105,20 @@ FExrImgMediaReader::EReadResult FExrImgMediaReader::ReadTilesCustom
 				break;
 			}
 		}
-		const uint16 Padding = bCustomExr ? 0 : FExrReader::TILE_PADDING;
-		const int64 TileByteStride = PixelSize * TileDim.X * TileDim.Y + Padding;
-		const int StartTileIndex = TileRow * FullTexDimInTiles.X + TileRegion.Min.X;
-		const int EndTileIndex = TileRow * FullTexDimInTiles.X + TileRegion.Max.X;
+		const uint16 Padding = ConverterParams->bCustomExr ? 0 : FExrReader::TILE_PADDING;
+		const int64 TileByteStride = ConverterParams->PixelSize * ConverterParams->TileDimWithBorders.X * ConverterParams->TileDimWithBorders.Y + Padding;
+		const int StartTileIndex = TileRow * DimensionInTiles.X + TileRegion.Min.X;
+		const int EndTileIndex = TileRow * DimensionInTiles.X + TileRegion.Max.X;
 		const int BytesToRead = (EndTileIndex - StartTileIndex) * TileByteStride;
 
-		if (bCustomExr)
+		if (ConverterParams->bCustomExr)
 		{
-			int64 TileStride = TileDim.X * TileDim.Y * PixelSize;
-			ChunkReader.SeekTileWithinFileCustom(StartTileIndex, TileStride, CurrentBufferPos);
+			int64 TileStride = ConverterParams->TileDimWithBorders.X * ConverterParams->TileDimWithBorders.Y * ConverterParams->PixelSize;
+			ChunkReader.SeekTileWithinFileCustom(StartTileIndex, TileStride, ConverterParams->bMipsInSeparateFiles ? 0 : CurrentMipLevel, CurrentBufferPos);
 		}
 		else
 		{
-			ChunkReader.SeekTileWithinFile(StartTileIndex, FullTexDimInTiles, CurrentBufferPos);
+			ChunkReader.SeekTileWithinFile(StartTileIndex, DimensionInTiles, ConverterParams->bMipsInSeparateFiles ? 0 : CurrentMipLevel, CurrentBufferPos);
 		}
 
 		if (!ChunkReader.ReadExrImageChunk(reinterpret_cast<char*>(Buffer) + CurrentBufferPos, BytesToRead))
@@ -259,8 +260,14 @@ bool FExrImgMediaReader::ReadFrame(int32 FrameId, const TMap<int32, FImgMediaTil
 						FMath::Min((int32)CurrentTileSelection.BottomRightY, FMath::CeilToInt(float(OutFrame->Info.NumTiles.Y) / MipLevelDiv)));
 					FIntPoint NumTiles = OutFrame->Info.NumTiles / MipLevelDiv;
 					int32 PixelSize = sizeof(uint16) * OutFrame->Info.NumChannels;
-
-					EReadResult ReadResult = ReadTilesCustom((uint16*)MipDataPtr, Image, FrameId, TileRegion, NumTiles, OutFrame->Info.TileDimensions, PixelSize, true /*bCustomExr*/);
+					TSharedPtr<FSampleConverterParameters> ConverterParams = MakeShared<FSampleConverterParameters>();
+					ConverterParams->FrameInfo = OutFrame->Info;
+					ConverterParams->PixelSize = sizeof(uint16) * ConverterParams->FrameInfo.NumChannels;
+					ConverterParams->TileDimWithBorders = OutFrame->Info.TileDimensions + OutFrame->Info.TileBorder * 2;
+					ConverterParams->NumMipLevels = Loader->GetNumMipLevels();
+					ConverterParams->bCustomExr = OutFrame->Info.FormatName == TEXT("EXR CUSTOM");
+
+					EReadResult ReadResult = ReadTilesCustom((uint16*)MipDataPtr, Image, FrameId, TileRegion, ConverterParams, CurrentMipLevel);
 					if (ReadResult != Fail)
 					{
 						OutFrame->MipTilesPresent.Emplace(CurrentMipLevel, CurrentTileSelection);
@@ -394,6 +401,7 @@ bool FExrImgMediaReader::GetInfo(const FString& FilePath, FImgMediaFrameInfo& Ou
 	OutInfo.Srgb = false;
 	OutInfo.UncompressedSize = HeaderReader.GetUncompressedSize();
 	OutInfo.NumChannels = HeaderReader.GetNumChannels();
+	OutInfo.NumMipLevels = 1;
 
 	int32 CustomFormat = 0;
 	HeaderReader.GetIntAttribute(IImgMediaModule::CustomFormatAttributeName.Resolve().ToString(), CustomFormat);
@@ -418,6 +426,10 @@ bool FExrImgMediaReader::GetInfo(const FString& FilePath, FImgMediaFrameInfo& Ou
 	if (OutInfo.bHasTiles)
 	{
 		OutInfo.NumTiles = FIntPoint(OutInfo.Dim.X / (OutInfo.TileDimensions.X + OutInfo.TileBorder * 2), OutInfo.Dim.Y / (OutInfo.TileDimensions.Y + OutInfo.TileBorder * 2));
+		if (HeaderReader.ContainsMips())
+		{
+			OutInfo.NumMipLevels = HeaderReader.CalculateNumMipLevels(OutInfo.NumTiles);
+		}
 	}
 	else
 	{
@@ -425,6 +437,7 @@ bool FExrImgMediaReader::GetInfo(const FString& FilePath, FImgMediaFrameInfo& Ou
 		OutInfo.NumTiles = FIntPoint(1, 1);
 	}
 
+
 	return (OutInfo.UncompressedSize > 0) && (OutInfo.Dim.GetMin() > 0);
 }
 
diff --git a/Engine/Plugins/Media/ImgMedia/Source/ImgMedia/Private/Readers/ExrImgMediaReader.h b/Engine/Plugins/Media/ImgMedia/Source/ImgMedia/Private/Readers/ExrImgMediaReader.h
--- a/Engine/Plugins/Media/ImgMedia/Source/ImgMedia/Private/Readers/ExrImgMediaReader.h
+++ b/Engine/Plugins/Media/ImgMedia/Source/ImgMedia/Private/Readers/ExrImgMediaReader.h
@@ -49,6 +49,21 @@ protected:
 		Cancelled
 	};
 
+	/* 
+	* These are all the required parameters to produce Buffer to Texture converter. 
+	*/
+	struct FSampleConverterParameters
+	{
+		FImgMediaFrameInfo FrameInfo;
+		FIntPoint FullResolution;
+		FIntPoint TileDimWithBorders;
+		TMap<int32, FIntRect> Viewports;
+		int32 NumMipLevels;
+		int32 PixelSize;
+		bool bCustomExr;
+		bool bMipsInSeparateFiles;
+	};
+
 	/**
 	 * Get the frame information from the given input file.
 	 *
@@ -62,7 +77,7 @@ protected:
 	 * Reads tiles from exr files in tile rows based on Tile region. If frame is pending for cancelation
 	 * stops reading tiles at the current tile row.
 	*/
-	EReadResult ReadTilesCustom(uint16* Buffer, const FString& ImagePath, int32 FrameId, const FIntRect& TileRegion, const FIntPoint& FullTexDimInTiles, const FIntPoint& TileDim, const int32 PixelSize, const bool bCustomExr);
+	EReadResult ReadTilesCustom(uint16* Buffer, const FString& ImagePath, int32 FrameId, const FIntRect& TileRegion, TSharedPtr<FSampleConverterParameters> ConverterParams, const int32 CurrentMipLevel);
 
 	/**
 	 * Sets parameters of our custom format images.
diff --git a/Engine/Plugins/Media/ImgMedia/Source/ImgMedia/Private/Readers/ExrImgMediaReaderGpu.cpp b/Engine/Plugins/Media/ImgMedia/Source/ImgMedia/Private/Readers/ExrImgMediaReaderGpu.cpp
--- a/Engine/Plugins/Media/ImgMedia/Source/ImgMedia/Private/Readers/ExrImgMediaReaderGpu.cpp
+++ b/Engine/Plugins/Media/ImgMedia/Source/ImgMedia/Private/Readers/ExrImgMediaReaderGpu.cpp
@@ -157,7 +157,7 @@ bool FExrImgMediaReaderGpu::ReadFrame(int32 FrameId, const TMap<int32, FImgMedia
 	ConverterParams->TileDimWithBorders = TileDim + OutFrame->Info.TileBorder * 2;
 	ConverterParams->NumMipLevels = Loader->GetNumMipLevels();
 	ConverterParams->bCustomExr = OutFrame->Info.FormatName == TEXT("EXR CUSTOM");
-
+	ConverterParams->bMipsInSeparateFiles = Loader->MipsInSeparateFiles();
 	{
 		// Loop over all mips.
 		FIntPoint CurrentMipDim = ConverterParams->FullResolution;
@@ -181,8 +181,8 @@ bool FExrImgMediaReaderGpu::ReadFrame(int32 FrameId, const TMap<int32, FImgMedia
 			if (ReadThisMip)
 			{
 				FIntRect& Viewport = ConverterParams->Viewports.Add(CurrentMipLevel);
-				Viewport.Min = FIntPoint(ConverterParams->TileDimWithBorders.X * CurrentTileSelection.TopLeftX, ConverterParams->TileDimWithBorders.X * CurrentTileSelection.TopLeftY);
-				Viewport.Max = FIntPoint(ConverterParams->TileDimWithBorders.Y * CurrentTileSelection.BottomRightX, ConverterParams->TileDimWithBorders.Y * CurrentTileSelection.BottomRightY);
+				Viewport.Min = FIntPoint(ConverterParams->TileDimWithBorders.X * CurrentTileSelection.TopLeftX, ConverterParams->TileDimWithBorders.Y * CurrentTileSelection.TopLeftY);
+				Viewport.Max = FIntPoint(ConverterParams->TileDimWithBorders.X * CurrentTileSelection.BottomRightX, ConverterParams->TileDimWithBorders.Y * CurrentTileSelection.BottomRightY);
 				Viewport.Clip(FIntRect(FIntPoint::ZeroValue, CurrentMipDim));
 
 				const SIZE_T BufferSize = GetBufferSize(CurrentMipDim, ConverterParams->FrameInfo.NumChannels, bHasTiles, OutFrame->Info.NumTiles / MipLevelDiv, ConverterParams->bCustomExr);
@@ -190,8 +190,8 @@ bool FExrImgMediaReaderGpu::ReadFrame(int32 FrameId, const TMap<int32, FImgMedia
 				BufferData = AllocateGpuBufferFromPool(BufferSize);
 				uint16* MipDataPtr = static_cast<uint16*>(BufferData->MappedBuffer);
 
-				// Get for our frame/mip level.
-				FString ImagePath = Loader->GetImagePath(FrameId, CurrentMipLevel);
+				// Get highest resolution mip level path.
+				FString ImagePath = Loader->GetImagePath(FrameId, ConverterParams->bMipsInSeparateFiles ? CurrentMipLevel : 0);
 				
 				EReadResult ReadResult = Fail;
 
@@ -206,9 +206,8 @@ bool FExrImgMediaReaderGpu::ReadFrame(int32 FrameId, const TMap<int32, FImgMedia
 							(int32)CurrentTileSelection.TopLeftY,
 							FMath::Min((int32)CurrentTileSelection.BottomRightX, FMath::CeilToInt(float(OutFrame->Info.NumTiles.X) / MipLevelDiv)),
 							FMath::Min((int32)CurrentTileSelection.BottomRightY, FMath::CeilToInt(float(OutFrame->Info.NumTiles.Y) / MipLevelDiv)));
-						FIntPoint NumTiles = OutFrame->Info.NumTiles / MipLevelDiv;
 
-						ReadResult = ReadTilesCustom(MipDataPtr, ImagePath, FrameId, TileRegion, NumTiles, ConverterParams->TileDimWithBorders, ConverterParams->PixelSize, ConverterParams->bCustomExr);
+						ReadResult = ReadTilesCustom(MipDataPtr, ImagePath, FrameId, TileRegion, ConverterParams, CurrentMipLevel);
 					}
 					else
 					{
@@ -288,8 +287,9 @@ FExrImgMediaReaderGpu::EReadResult FExrImgMediaReaderGpu::ReadInChunks(uint16* B
 	int32 CurrentBufferPos = 0;
 	FExrReader ChunkReader;
 
-
-	if (!ChunkReader.OpenExrAndPrepareForPixelReading(ImagePath, Dim.Y))
+	// Since ReadInChunks is only utilized for exr files without tiles and mips, Num Mip levels is always 1.
+	const int32 NumLevels = 1;
+	if (!ChunkReader.OpenExrAndPrepareForPixelReading(ImagePath, Dim.Y, NumLevels))
 	{
 		return Fail;
 	}
diff --git a/Engine/Plugins/Media/ImgMedia/Source/ImgMedia/Private/Readers/ExrImgMediaReaderGpu.h b/Engine/Plugins/Media/ImgMedia/Source/ImgMedia/Private/Readers/ExrImgMediaReaderGpu.h
--- a/Engine/Plugins/Media/ImgMedia/Source/ImgMedia/Private/Readers/ExrImgMediaReaderGpu.h
+++ b/Engine/Plugins/Media/ImgMedia/Source/ImgMedia/Private/Readers/ExrImgMediaReaderGpu.h
@@ -72,20 +72,6 @@ public:
 
 protected:
 
-	/* 
-	* These are all the required parameters to produce Buffer to Texture converter. 
-	*/
-	struct FSampleConverterParameters
-	{
-		FImgMediaFrameInfo FrameInfo;
-		FIntPoint FullResolution;
-		FIntPoint TileDimWithBorders;
-		TMap<int32, FIntRect> Viewports;
-		int32 NumMipLevels;
-		int32 PixelSize;
-		bool bCustomExr;
-	};
-
 	/** 
 	 * This function reads file in 16 MB chunks and if it detects that
 	 * Frame is pending for cancellation stops reading the file and returns false.
diff --git a/Engine/Plugins/Media/ImgMedia/Source/ImgMedia/Private/Readers/IImgMediaReader.h b/Engine/Plugins/Media/ImgMedia/Source/ImgMedia/Private/Readers/IImgMediaReader.h
--- a/Engine/Plugins/Media/ImgMedia/Source/ImgMedia/Private/Readers/IImgMediaReader.h
+++ b/Engine/Plugins/Media/ImgMedia/Source/ImgMedia/Private/Readers/IImgMediaReader.h
@@ -48,6 +48,9 @@ struct FImgMediaFrameInfo
 
 	/** Tile border size in texels. This is required for more elaborate tile texel sampling. */
 	int32 TileBorder;
+
+	/** Number of Mip Levels. */
+	int32 NumMipLevels;
 };
 
 
diff --git a/Engine/Plugins/Media/ImgMedia/Source/OpenExrWrapper/Private/OpenExrWrapperContextReader.cpp b/Engine/Plugins/Media/ImgMedia/Source/OpenExrWrapper/Private/OpenExrWrapperContextReader.cpp
--- a/Engine/Plugins/Media/ImgMedia/Source/OpenExrWrapper/Private/OpenExrWrapperContextReader.cpp
+++ b/Engine/Plugins/Media/ImgMedia/Source/OpenExrWrapper/Private/OpenExrWrapperContextReader.cpp
@@ -136,6 +136,32 @@ int32 FOpenExrHeaderReader::GetNumChannels() const
 	return Attribute->chlist->num_channels;
 }
 
+bool FOpenExrHeaderReader::ContainsMips() const
+{
+	const exr_attribute_t* Attribute = nullptr;
+	exr_result_t Result = exr_get_attribute_by_name(*((exr_context_t*)FileContext.Get()), 0, EXR_ATTRIBUTE_TILEDESC, &Attribute);
+	if (Result == EXR_ERR_SUCCESS && Attribute)
+	{
+		exr_tile_level_mode_t TileMipMode = EXR_GET_TILE_LEVEL_MODE(*Attribute->tiledesc);
+		return TileMipMode == EXR_TILE_MIPMAP_LEVELS;
+	}
+	return false;
+}
+
+int32 FOpenExrHeaderReader::CalculateNumMipLevels(const FIntPoint& NumTiles) const
+{
+	const exr_attribute_t* Attribute = nullptr;
+	exr_result_t Result = exr_get_attribute_by_name(*((exr_context_t*)FileContext.Get()), 0, EXR_ATTRIBUTE_TILEDESC, &Attribute);
+	if (Result == EXR_ERR_SUCCESS && Attribute)
+	{
+		exr_tile_round_mode_t MipRoundMode = EXR_GET_TILE_ROUND_MODE(*Attribute->tiledesc);
+		int32 MinTileRes = FMath::Min(NumTiles.X, NumTiles.Y);
+		int32 NumMipLevels = (exr_tile_round_mode_t::EXR_TILE_ROUND_DOWN) ? FMath::FloorLog2(MinTileRes) : FMath::CeilLogTwo(MinTileRes);
+		return NumMipLevels + 1;
+	}
+	return 1;
+}
+
 bool FOpenExrHeaderReader::GetTileSize(FIntPoint& OutTileSize) const
 {
 	const exr_attribute_t* Attribute = nullptr;
diff --git a/Engine/Plugins/Media/ImgMedia/Source/OpenExrWrapper/Public/OpenExrWrapper.h b/Engine/Plugins/Media/ImgMedia/Source/OpenExrWrapper/Public/OpenExrWrapper.h
--- a/Engine/Plugins/Media/ImgMedia/Source/OpenExrWrapper/Public/OpenExrWrapper.h
+++ b/Engine/Plugins/Media/ImgMedia/Source/OpenExrWrapper/Public/OpenExrWrapper.h
@@ -70,8 +70,10 @@ public:
 	FFrameRate GetFrameRate(const FFrameRate& DefaultValue) const;
 	int32 GetUncompressedSize() const;
 	int32 GetNumChannels() const;
+	bool ContainsMips() const;
+	int32 CalculateNumMipLevels(const FIntPoint& NumTiles) const;
 
-	// Gets tile dimensions. Returns false if image has no tiles.
+	/** Gets tile dimensions. Returns false if image has no tiles. */
 	bool GetTileSize(FIntPoint& OutTileSize) const;
 
 	bool HasInputFile() const;
