ImgMediaPlayer ExrReaderGPU: Reading custom Exr files with embedded mips.

diff --git a/Engine/Plugins/Media/ImgMedia/Source/ExrReaderGpu/Private/ExrReaderGpu.cpp b/Engine/Plugins/Media/ImgMedia/Source/ExrReaderGpu/Private/ExrReaderGpu.cpp
--- a/Engine/Plugins/Media/ImgMedia/Source/ExrReaderGpu/Private/ExrReaderGpu.cpp
+++ b/Engine/Plugins/Media/ImgMedia/Source/ExrReaderGpu/Private/ExrReaderGpu.cpp
@@ -251,7 +251,7 @@ bool FExrReader::GenerateTextureData(uint16* Buffer, int32 BufferSize, FString F
 	return true;
 }
 
-bool FExrReader::OpenExrAndPrepareForPixelReading(FString FilePath, int32 NumOffsets, int32 NumLevels)
+bool FExrReader::OpenExrAndPrepareForPixelReading(FString FilePath, int32 NumOffsets, int32 NumLevels, bool bCustomExr)
 {
 	if (FileHandle != nullptr)
 	{
@@ -279,7 +279,7 @@ bool FExrReader::OpenExrAndPrepareForPixelReading(FString FilePath, int32 NumOff
 		ELineOrder LineOrder = INCREASING_Y;
 		for (int Level = 0; Level < NumLevels; Level++)
 		{
-			LineOrTileOffsetsPerLevel[Level].SetNum(NumOffsets >> (Level * 2));
+			LineOrTileOffsetsPerLevel[Level].SetNum(bCustomExr ? 1 : NumOffsets >> (Level * 2));
 			ReadLineOrTileOffsets(FileHandle, LineOrder, LineOrTileOffsetsPerLevel[Level]);
 		}
 	}
diff --git a/Engine/Plugins/Media/ImgMedia/Source/ExrReaderGpu/Public/ExrReaderGpu.h b/Engine/Plugins/Media/ImgMedia/Source/ExrReaderGpu/Public/ExrReaderGpu.h
--- a/Engine/Plugins/Media/ImgMedia/Source/ExrReaderGpu/Public/ExrReaderGpu.h
+++ b/Engine/Plugins/Media/ImgMedia/Source/ExrReaderGpu/Public/ExrReaderGpu.h
@@ -64,7 +64,7 @@ public:
 	* Reading in chunks allows the process of reading to be canceled midway through reading. 
 	* This function reads and discards the header and keeps the pointers to the scanlines. 
 	*/
-	bool OpenExrAndPrepareForPixelReading(FString FilePath, int32 NumOffsets, int32 NumLevels);
+	bool OpenExrAndPrepareForPixelReading(FString FilePath, int32 NumOffsets, int32 NumLevels, bool bCustomExr = false);
 
 	/** 
 	* Read a chunk of an Exr file previously open via OpenExrAndPrepareForPixelReading.
diff --git a/Engine/Plugins/Media/ImgMedia/Source/ImgMedia/Private/Readers/ExrImgMediaReader.cpp b/Engine/Plugins/Media/ImgMedia/Source/ImgMedia/Private/Readers/ExrImgMediaReader.cpp
--- a/Engine/Plugins/Media/ImgMedia/Source/ImgMedia/Private/Readers/ExrImgMediaReader.cpp
+++ b/Engine/Plugins/Media/ImgMedia/Source/ImgMedia/Private/Readers/ExrImgMediaReader.cpp
@@ -88,7 +88,7 @@ FExrImgMediaReader::EReadResult FExrImgMediaReader::ReadTilesCustom
 	FIntPoint DimensionInTiles = ConverterParams->FrameInfo.NumTiles / MipLevelDiv;
 	int32 NumTiles = ConverterParams->bCustomExr ? 1 : ConverterParams->FrameInfo.NumTiles.X * ConverterParams->FrameInfo.NumTiles.Y;
 	int32 NumMipLevels = ConverterParams->bMipsInSeparateFiles ? 1 : ConverterParams->NumMipLevels;
-	if (!ChunkReader.OpenExrAndPrepareForPixelReading(ImagePath, NumTiles, NumMipLevels))
+	if (!ChunkReader.OpenExrAndPrepareForPixelReading(ImagePath, NumTiles, NumMipLevels, ConverterParams->bCustomExr))
 	{
 		return Fail;
 	}
@@ -113,7 +113,7 @@ FExrImgMediaReader::EReadResult FExrImgMediaReader::ReadTilesCustom
 
 		if (ConverterParams->bCustomExr)
 		{
-			int64 TileStride = ConverterParams->TileDimWithBorders.X * ConverterParams->TileDimWithBorders.Y * ConverterParams->PixelSize;
+			const int64 TileStride = ConverterParams->TileDimWithBorders.X * ConverterParams->TileDimWithBorders.Y * ConverterParams->PixelSize;
 			ChunkReader.SeekTileWithinFileCustom(StartTileIndex, TileStride, ConverterParams->bMipsInSeparateFiles ? 0 : CurrentMipLevel, CurrentBufferPos);
 		}
 		else
