ImgMedia ExrReader: Fixed an issue of Exr mip rounding mode not being used correctly.

diff --git a/Engine/Plugins/Media/ImgMedia/Source/OpenExrWrapper/Private/OpenExrWrapperContextReader.cpp b/Engine/Plugins/Media/ImgMedia/Source/OpenExrWrapper/Private/OpenExrWrapperContextReader.cpp
--- a/Engine/Plugins/Media/ImgMedia/Source/OpenExrWrapper/Private/OpenExrWrapperContextReader.cpp
+++ b/Engine/Plugins/Media/ImgMedia/Source/OpenExrWrapper/Private/OpenExrWrapperContextReader.cpp
@@ -156,7 +156,7 @@ int32 FOpenExrHeaderReader::CalculateNumMipLevels(const FIntPoint& NumTiles) con
 	{
 		exr_tile_round_mode_t MipRoundMode = EXR_GET_TILE_ROUND_MODE(*Attribute->tiledesc);
 		int32 MinTileRes = FMath::Min(NumTiles.X, NumTiles.Y);
-		int32 NumMipLevels = (exr_tile_round_mode_t::EXR_TILE_ROUND_DOWN) ? FMath::FloorLog2(MinTileRes) : FMath::CeilLogTwo(MinTileRes);
+		int32 NumMipLevels = (MipRoundMode == exr_tile_round_mode_t::EXR_TILE_ROUND_DOWN) ? FMath::FloorLog2(MinTileRes) : FMath::CeilLogTwo(MinTileRes);
 		return NumMipLevels + 1;
 	}
 	return 1;
