ImgMedia: Improvements to visible mips and tiles calculations.

diff --git a/Engine/Plugins/Media/ImgMedia/Source/ImgMedia/Private/Assets/ImgMediaMipMapInfo.cpp b/Engine/Plugins/Media/ImgMedia/Source/ImgMedia/Private/Assets/ImgMediaMipMapInfo.cpp
--- a/Engine/Plugins/Media/ImgMedia/Source/ImgMedia/Private/Assets/ImgMediaMipMapInfo.cpp
+++ b/Engine/Plugins/Media/ImgMedia/Source/ImgMedia/Private/Assets/ImgMediaMipMapInfo.cpp
@@ -19,45 +19,403 @@
 
 DECLARE_CYCLE_STAT(TEXT("ImgMedia MipMap Update Cache"), STAT_ImgMedia_MipMapUpdateCache, STATGROUP_Media);
 
-FImgMediaMipMapInfo::FImgMediaMipMapInfo()
-	: MipLevel0Distance(-1.0f)
-	, bIsMipLevel0DistanceSetManually(false)
-	, ViewportDistAdjust(0.0f)
-	, CachedMipLevel(0)
-	, bIsCachedMipLevelValid(false)
+FImgMediaMipMapObjectInfo::FImgMediaMipMapObjectInfo(UMeshComponent* InMeshComponent, float InLODBias)
+	: MeshComponent(InMeshComponent)
+	, LODBias(InLODBias)
 {
+
 }
 
-FImgMediaMipMapInfo::~FImgMediaMipMapInfo()
+UMeshComponent* FImgMediaMipMapObjectInfo::GetMeshComponent() const
 {
-	for (FImgMediaMipMapObjectInfo* Info : Objects)
+	return MeshComponent.Get();
+}
+
+bool FImgMediaMipMapObjectInfo::CalculateVisibleTiles(const TArray<FImgMediaMipMapCameraInfo>& InCameraInfos, const FSequenceInfo& InSequenceInfo, TMap<int32, FImgMediaTileSelection>& VisibleTiles) const
+{
+	for (int32 MipLevel = 0; MipLevel < InSequenceInfo.NumMipLevels; ++MipLevel)
 	{
-		delete Info;
+		FImgMediaTileSelection& Selection = VisibleTiles.FindOrAdd(MipLevel);
+
+		int32 NumTilesX = FMath::Max(1, InSequenceInfo.NumTiles.X >> MipLevel);
+		int32 NumTilesY = FMath::Max(1, InSequenceInfo.NumTiles.Y >> MipLevel);
+		Selection.SetVisibleRegion(0u, 0u, IntCastChecked<uint16>(NumTilesX), IntCastChecked<uint16>(NumTilesY));
 	}
-	Objects.Empty();
+	return true;
+}
+
+namespace {
+	// Minimalized version of FSceneView::ProjectWorldToScreen
+	FORCEINLINE bool ProjectWorldToScreenFast(const FVector& WorldPosition, const FIntRect& ViewRect, const FMatrix& ViewProjectionMatrix, FVector2D& out_ScreenPos)
+	{
+		FPlane Result = ViewProjectionMatrix.TransformFVector4(FVector4(WorldPosition, 1.f));
+		if (Result.W > 0.0f)
+		{
+			float NormalizedX = (Result.X / (Result.W * 2.f)) + 0.5f;
+			float NormalizedY = 1.f - (Result.Y / (Result.W * 2.f)) - 0.5f;
+			out_ScreenPos = FVector2D(NormalizedX * (float)ViewRect.Width(), NormalizedY * (float)ViewRect.Height());
+
+			return true;
+		}
+
+		return false;
+	}
+
+	// Approximates hardware mip level selection.
+	bool CalculateMipLevel(const FImgMediaMipMapCameraInfo& CameraInfo, const FVector& TexelWS, const FVector& OffsetXWS, const FVector& OffsetYWS, float& OutMipLevel)
+	{
+		FVector2D TexelScreenSpace[3];
+
+		bool bValid = true;
+		bValid &= ProjectWorldToScreenFast(TexelWS, CameraInfo.ViewportRect, CameraInfo.ViewProjectionMatrix, TexelScreenSpace[0]);
+		bValid &= ProjectWorldToScreenFast(TexelWS + OffsetXWS, CameraInfo.ViewportRect, CameraInfo.ViewProjectionMatrix, TexelScreenSpace[1]);
+		bValid &= ProjectWorldToScreenFast(TexelWS + OffsetYWS, CameraInfo.ViewportRect, CameraInfo.ViewProjectionMatrix, TexelScreenSpace[2]);
+
+		if (bValid)
+		{
+			float DistX = FVector2D::DistSquared(TexelScreenSpace[0], TexelScreenSpace[1]);
+			float DistY = FVector2D::DistSquared(TexelScreenSpace[0], TexelScreenSpace[2]);
+			OutMipLevel = FMath::Max(0.5f * (float)FMath::Log2(1.0f / FMath::Max(DistX, DistY)), 0.0f); // ~ log2(sqrt(delta))
+		}
+
+		return bValid;
+	}
+
+	class FPlaneObjectInfo : public FImgMediaMipMapObjectInfo
+	{
+	public:
+		FPlaneObjectInfo(UMeshComponent* InMeshComponent, float InLODBias = 0.0f)
+			: FImgMediaMipMapObjectInfo(InMeshComponent, InLODBias)
+			, PlaneSize(FVector::ZeroVector)
+		{
+			// Get size of object.
+			if (MeshComponent != nullptr)
+			{
+				PlaneSize = 2.0f * MeshComponent->CalcLocalBounds().BoxExtent;
+
+				PlaneVertices[0] = 0.5f * FVector(-PlaneSize.X, -PlaneSize.Y, 0);
+				PlaneVertices[1] = 0.5f * FVector(-PlaneSize.X, PlaneSize.Y, 0);
+				PlaneVertices[2] = 0.5f * FVector(PlaneSize.X, -PlaneSize.Y, 0);
+				PlaneVertices[3] = 0.5f * FVector(PlaneSize.X, PlaneSize.Y, 0);
+			}
+			else
+			{
+				UE_LOG(LogImgMedia, Error, TEXT("FPlaneImgMediaMipMapObjectInfo is missing its plane mesh component."));
+			}
+		}
+
+		bool CalculateVisibleTiles(const TArray<FImgMediaMipMapCameraInfo>& InCameraInfos, const FSequenceInfo& InSequenceInfo, TMap<int32, FImgMediaTileSelection>& VisibleTiles) const override
+		{
+			UMeshComponent* Mesh = MeshComponent.Get();
+			if (Mesh == nullptr)
+			{
+				return false;
+			}
+
+			// To avoid calculating tile corner mip levels multiple times over, we cache them in this array.
+			CornerMipLevelsCached.SetNum((InSequenceInfo.NumTiles.X + 1) * (InSequenceInfo.NumTiles.Y + 1));
+			for (float& Level : CornerMipLevelsCached)
+			{
+				Level = -1.0f;
+			}
+
+			const FTransform MeshTransform = Mesh->GetComponentTransform();
+			const FVector MeshScale = Mesh->GetComponentScale();
+
+			FVector PlaneCornerWS = MeshTransform.TransformPosition(-0.5f * PlaneSize);
+			FVector DirXWS = MeshTransform.TransformVector(FVector(PlaneSize.X, 0, 0));
+			FVector DirYWS = MeshTransform.TransformVector(FVector(0, PlaneSize.Y, 0));
+			FVector TexelOffsetXWS = MeshTransform.TransformVector(FVector(PlaneSize.X / InSequenceInfo.Dim.X, 0, 0));
+			FVector TexelOffsetYWS = MeshTransform.TransformVector(FVector(0, PlaneSize.Y / InSequenceInfo.Dim.Y, 0));
+
+			for (const FImgMediaMipMapCameraInfo& CameraInfo : InCameraInfos)
+			{
+				// Get frustum.
+				FConvexVolume ViewFrustum;
+				GetViewFrustumBounds(ViewFrustum, CameraInfo.ViewProjectionMatrix, false, false);
+
+				int32 MaxLevel = InSequenceInfo.NumMipLevels - 1;
+				FIntPoint CurrentNumTiles;
+				CurrentNumTiles.X = FMath::Max(1, InSequenceInfo.NumTiles.X >> MaxLevel);
+				CurrentNumTiles.Y = FMath::Max(1, InSequenceInfo.NumTiles.Y >> MaxLevel);
+
+				// Starting with tiles at the highest mip level
+				TQueue<FIntVector> Tiles;
+				for (int32 TileY = 0; TileY < CurrentNumTiles.Y; ++TileY)
+				{
+					for (int32 TileX = 0; TileX < CurrentNumTiles.X; ++TileX)
+					{
+						Tiles.Enqueue(FIntVector(TileX, TileY, MaxLevel));
+					}
+				}
+
+				// Process all visible tiles with a (quadtree) breadth-first search
+				while (!Tiles.IsEmpty())
+				{
+					FIntVector Tile;
+					Tiles.Dequeue(Tile);
+
+					int32 CurrentMipLevel = Tile.Z;
+					// Calculate the number of tiles at this mip level
+					CurrentNumTiles.X = FMath::Max(1, InSequenceInfo.NumTiles.X >> CurrentMipLevel);
+					CurrentNumTiles.Y = FMath::Max(1, InSequenceInfo.NumTiles.Y >> CurrentMipLevel);
+
+					// Calculate the tile location in world-space
+					float StepX = float(Tile.X + 0.5f) / CurrentNumTiles.X;
+					float StepY = float(Tile.Y + 0.5f) / CurrentNumTiles.Y;
+					FVector TileCenterWS = PlaneCornerWS + (DirXWS * StepX + DirYWS * StepY);
+
+					// Calculate the tile radius in world space
+					FVector TileSizeWS = (PlaneSize * MeshScale) / FVector(CurrentNumTiles.X, CurrentNumTiles.Y, 1);
+					float TileRadiusInWS = 0.5f * (float)FMath::Sqrt(2 * FMath::Square(TileSizeWS.GetAbsMax()));
+
+					// Now we check if tile spherical bounds are in view.
+					if (ViewFrustum.IntersectSphere(TileCenterWS, TileRadiusInWS))
+					{
+						// Calculate the visible mip level range over all tile corners.
+						int32 NumVisibleCorners = 0;
+						FIntVector2 MipLevelRange = FIntVector2(TNumericLimits<int32>::Max(), 0);
+						for (int32 CornerY = 0; CornerY < 2; ++CornerY)
+						{
+							for (int32 CornerX = 0; CornerX < 2; ++CornerX)
+							{
+								float CalculatedLevel;
+								int32 TileCornerX = Tile.X + CornerX;
+								int32 TileCornerY = Tile.Y + CornerY;
+
+								// First we query the cached corner mip levels.
+								int32 MaxCornerX = InSequenceInfo.NumTiles.X + 1;
+								int32 MaxCornerY = InSequenceInfo.NumTiles.Y + 1;
+								FIntPoint BaseLevelCorner;
+								BaseLevelCorner.X = FMath::Clamp(TileCornerX << CurrentMipLevel, 0, InSequenceInfo.NumTiles.X);
+								BaseLevelCorner.Y = FMath::Clamp(TileCornerY << CurrentMipLevel, 0, InSequenceInfo.NumTiles.Y);
+								bool bValidLevel = GetCachedMipLevel(BaseLevelCorner.X, BaseLevelCorner.Y, MaxCornerX, CalculatedLevel);
+
+								// If not found, calculate and cache it.
+								if (!bValidLevel)
+								{
+									float CornerStepX = TileCornerX / (float)CurrentNumTiles.X;
+									float CornerStepY = TileCornerY / (float)CurrentNumTiles.Y;
+									FVector CornersWS = PlaneCornerWS + (DirXWS * CornerStepX + DirYWS * CornerStepY);
+
+									if (CalculateMipLevel(CameraInfo, CornersWS, TexelOffsetXWS, TexelOffsetYWS, CalculatedLevel))
+									{
+										CalculatedLevel += LODBias + CameraInfo.MaterialTextureMipBias;
+
+										SetCachedMipLevel(BaseLevelCorner.X, BaseLevelCorner.Y, MaxCornerX, CalculatedLevel);
+										bValidLevel = true;
+									}
+								}
+								
+								if (bValidLevel)
+								{
+									MipLevelRange[0] = FMath::Min(MipLevelRange[0], FMath::Min((int32)CalculatedLevel, MaxLevel));
+									MipLevelRange[1] = FMath::Max(MipLevelRange[1], FMath::Min(FMath::CeilToInt32(CalculatedLevel), MaxLevel));
+									NumVisibleCorners++;
+								}
+							}
+						}
+
+						// As an approximation, we force the lowest mip to 0 if only some corners are behind camera.
+						if (NumVisibleCorners > 0 && NumVisibleCorners < 4)
+						{
+							MipLevelRange[0] = 0;
+						}
+
+						// If the lowest (calculated) mip level is below our current mip level, enqueue all 4 sub-tiles for further processing.
+						if (MipLevelRange[0] < CurrentMipLevel)
+						{
+							for (int32 SubY = 0; SubY < FMath::Min(InSequenceInfo.NumTiles.Y, 2); ++SubY)
+							{
+								for (int32 SubX = 0; SubX < FMath::Min(InSequenceInfo.NumTiles.X, 2); ++SubX)
+								{
+									FIntVector SubTile = FIntVector((Tile.X << 1) + SubX, (Tile.Y << 1) + SubY, CurrentMipLevel - 1);
+									Tiles.Enqueue(SubTile);
+								}
+							}
+						}
+
+						// If the highest (calculated) mip level equals or exceeds our current mip level, we register the tile as visible.
+						if (MipLevelRange[1] >= CurrentMipLevel)
+						{
+							VisibleTiles.FindOrAdd(CurrentMipLevel).Include(Tile.X, Tile.Y);
+						}
+#if false
+#if WITH_EDITOR
+						// Enable this to draw a sphere where each tile is.
+						Async(EAsyncExecution::TaskGraphMainThread, [TileCenterWS, TileRadiusInWS]()
+							{
+								UWorld* World = GEditor->GetEditorWorldContext().World();
+								DrawDebugSphere(World, TileCenterWS, TileRadiusInWS, 8, FColor::Red, false, 0.05f);
+							});
+#endif // WITH_EDITOR
+#endif // false
+					}
+				}
+			}
+			return true;
+		}
+
+	private:
+
+		/** Convenience function to get a cached calculated mip level (in mip0 tile address space). */
+		FORCEINLINE bool GetCachedMipLevel(int32 Address0X, int32 Address0Y, int32 RowSize, float& OutCalculatedLevel) const
+		{
+			const int32 Index = Address0Y * RowSize + Address0X;
+
+			if (CornerMipLevelsCached[Index] >= 0.0f)
+			{
+				OutCalculatedLevel = CornerMipLevelsCached[Index];
+				
+				return true;
+			}
+
+			return false;
+		}
+
+		/** Convenience function to cache a calculated mip level (in mip0 tile address space). */
+		FORCEINLINE void SetCachedMipLevel(int32 Address0X, int32 Address0Y, int32 RowSize, float InCalculatedLevel) const
+		{
+			CornerMipLevelsCached[Address0Y * RowSize + Address0X] = InCalculatedLevel;
+		}
+
+		/** Local size of this mesh component. */
+		FVector PlaneSize;
+
+		/** Four plane vertices (in local space). */
+		TStaticArray<FVector, 4> PlaneVertices;
+
+		/** Cached calculating mip levels (at mip0). */
+		mutable TArray<float> CornerMipLevelsCached;
+	};
+
+	// Reduce the selected tiles region for higher mip levels, always ensuring that all tiles at the base level are included.
+	FImgMediaTileSelection DownscaleTileSelection(const FImgMediaTileSelection& InTopLevelSelection, int32 InMipLevel)
+	{
+		if (InMipLevel <= 0)
+		{
+			return InTopLevelSelection;
+		}
+
+		const int32 MipLevelDiv = 1 << InMipLevel;
+
+		FImgMediaTileSelection OutSelection;
+		OutSelection.TopLeftX = (uint16)(InTopLevelSelection.TopLeftX / MipLevelDiv);
+		OutSelection.TopLeftY = (uint16)(InTopLevelSelection.TopLeftY / MipLevelDiv);
+		OutSelection.BottomRightX = (uint16)FMath::CeilToInt(float(InTopLevelSelection.BottomRightX) / MipLevelDiv);
+		OutSelection.BottomRightY = (uint16)FMath::CeilToInt(float(InTopLevelSelection.BottomRightY) / MipLevelDiv);
+
+		return OutSelection;
+	}
+
+	class FSphereObjectInfo : public FImgMediaMipMapObjectInfo
+	{
+	public:
+		using FImgMediaMipMapObjectInfo::FImgMediaMipMapObjectInfo;
+
+		bool CalculateVisibleTiles(const TArray<FImgMediaMipMapCameraInfo>& InCameraInfos, const FSequenceInfo& InSequenceInfo, TMap<int32, FImgMediaTileSelection>& VisibleTiles) const override
+		{
+			UMeshComponent* Mesh = MeshComponent.Get();
+			if (Mesh == nullptr)
+			{
+				return false;
+			}
+
+			const float DefaultSphereRadius = 160.0f;
+
+			for (const FImgMediaMipMapCameraInfo& CameraInfo : InCameraInfos)
+			{
+				// Analytical derivation of visible tiles from the view frustum, given a sphere presumed to be infinitely large
+				FConvexVolume ViewFrustum;
+				GetViewFrustumBounds(ViewFrustum, CameraInfo.ViewProjectionMatrix, false, false);
+				
+				// Include all tiles containted in the visible UV region
+				uint16 NumX = IntCastChecked<uint16>(InSequenceInfo.NumTiles.X);
+				uint16 NumY = IntCastChecked<uint16>(InSequenceInfo.NumTiles.Y);
+				for (uint16 TileY = 0; TileY < NumY; ++TileY)
+				{
+					for (uint16 TileX = 0; TileX < NumX; ++TileX)
+					{
+						FVector2D TileCornerUV = FVector2D((float)TileX / InSequenceInfo.NumTiles.X, (float)TileY / InSequenceInfo.NumTiles.Y);
+						
+						// Convert from latlong UV to spherical coordinates
+						FVector2D TileCornerSpherical = FVector2D(UE_PI * TileCornerUV.Y, UE_TWO_PI * TileCornerUV.X);
+						
+						// Adjust spherical coordinates to default sphere UVs
+						TileCornerSpherical.Y = -UE_HALF_PI - TileCornerSpherical.Y;
+
+						FVector TileCorner = TileCornerSpherical.SphericalToUnitCartesian() * DefaultSphereRadius;
+						TileCorner = Mesh->GetComponentTransform().TransformPosition(TileCorner);
+
+						// For each tile corner, we include all adjacent tiles
+						if (ViewFrustum.IntersectPoint(TileCorner))
+						{
+							VisibleTiles.FindOrAdd(0).Include(TileX, TileY);
+
+							int AdjacentX = TileX > 0 ? TileX - 1 : NumX - 1;
+							int AdjacentY = TileY > 0 ? TileY - 1 : NumY - 1;
+							VisibleTiles[0].Include(AdjacentX, TileY);
+							VisibleTiles[0].Include(TileX, AdjacentY);
+							VisibleTiles[0].Include(AdjacentX, AdjacentY);
+						}
+					}
+				}
+
+				if (VisibleTiles.Contains(0))
+				{
+					for (int32 Level = 1; Level < InSequenceInfo.NumMipLevels; ++Level)
+					{
+						FImgMediaTileSelection& Selection = VisibleTiles.FindOrAdd(Level);
+						Selection = DownscaleTileSelection(VisibleTiles[0], Level);
+					}
+				}
+			}
+
+			return true;
+		}
+	};
+
+} //end anonymous namespace
+
+FImgMediaMipMapInfo::FImgMediaMipMapInfo()
+	: bIsCacheValid(false)
+{
+}
+
+FImgMediaMipMapInfo::~FImgMediaMipMapInfo()
+{
+	ClearAllObjects();
 }
 
 void FImgMediaMipMapInfo::AddObject(AActor* InActor, float Width, float LODBias)
 {
 	if (InActor != nullptr)
 	{
-		FImgMediaMipMapObjectInfo* Info = new FImgMediaMipMapObjectInfo();
-		Info->Object = InActor;
+		UMeshComponent* MeshComponent = Cast<UMeshComponent>(InActor->FindComponentByClass(UMeshComponent::StaticClass()));
+		if (MeshComponent != nullptr)
+		{
+			FImgMediaMipMapObjectInfo* Info = nullptr;
 
-		// Calculate dist adjust from bias.
-		Info->DistAdjust = FMath::Pow(2.0f, LODBias);
+			if (UStaticMeshComponent* StaticMeshComponent = Cast<UStaticMeshComponent>(MeshComponent))
+			{
+				FString StaticMeshName = StaticMeshComponent->GetStaticMesh()->GetName();
+				if (StaticMeshName.Equals(TEXT("Plane")))
+				{
+					Info = new FPlaneObjectInfo(MeshComponent, LODBias);
+				}
+				else if (StaticMeshName.Equals(TEXT("Sphere")))
+				{
+					Info = new FSphereObjectInfo(MeshComponent, LODBias);
+				}
+			}
 
-		// Get size of object.
-		GetObjectSize(InActor, Info->Width, Info->Height);
+			if (Info == nullptr)
+			{
+				Info = new FImgMediaMipMapObjectInfo(MeshComponent, LODBias);
+			}
 
-		// Get dist adjust.
-		if (Width > 0.0f)
-		{
-			FImgMediaMipMapInfoManager& InfoManager = FImgMediaMipMapInfoManager::Get();
-			Info->DistAdjust *= InfoManager.GetRefObjectWidth() / Width;
+			Objects.Add(Info);
 		}
-
-		Objects.Add(Info);
 	}
 }
 
@@ -68,7 +426,7 @@ void FImgMediaMipMapInfo::RemoveObject(AActor* InActor)
 		for (int Index = 0; Index < Objects.Num(); ++Index)
 		{
 			FImgMediaMipMapObjectInfo* Info = Objects[Index];
-			if (InActor == Info->Object.Get())
+			if (InActor == Info->GetMeshComponent()->GetOuter())
 			{
 				Objects.RemoveAtSwap(Index);
 				delete Info;
@@ -103,40 +461,27 @@ void FImgMediaMipMapInfo::AddObjectsUsingThisMediaTexture(UMediaTexture* InMedia
 
 void FImgMediaMipMapInfo::ClearAllObjects()
 {
+	for (FImgMediaMipMapObjectInfo* Info : Objects)
+	{
+		delete Info;
+	}
 	Objects.Empty();
 }
 
-void FImgMediaMipMapInfo::SetMipLevelDistance(float Distance)
-{
-	MipLevel0Distance = Distance;
-	bIsMipLevel0DistanceSetManually = Distance >= 0.0f;
-	UpdateMipLevelDistances();
-}
 
-void FImgMediaMipMapInfo::SetTextureInfo(FName InSequenceName, int32 NumMipMaps,
-	int32 InNumTilesX, int32 InNumTilesY, const FIntPoint& Dim)
+void FImgMediaMipMapInfo::SetTextureInfo(FName InSequenceName, int32 InNumMipMaps,
+	const FIntPoint& InNumTiles, const FIntPoint& InSequenceDim)
 {
-	SequenceName = InSequenceName;
-	MipLevelDistances.SetNum(NumMipMaps);
-	NumTilesX = InNumTilesX;
-	NumTilesY = InNumTilesY;
+	SequenceInfo.Name = InSequenceName;
+	SequenceInfo.Dim = InSequenceDim;
 
-	// Set ideal distance for mip level 0 based on various factors.
-	if (bIsMipLevel0DistanceSetManually == false)
-	{
-		// Get reference distance.
-		FImgMediaMipMapInfoManager& InfoManager = FImgMediaMipMapInfoManager::Get();
-		MipLevel0Distance = InfoManager.GetMipLevel0Distance();
-
-		// Take our texture size into account.
-		float TextureAdjust = InfoManager.GetRefObjectTextureWidth() / Dim.X;
-		MipLevel0Distance *= TextureAdjust;
-	}
-
-	UpdateMipLevelDistances();
+	// To simplify logic, we assume we always have at least one mip level and one tile.
+	SequenceInfo.NumMipLevels = FMath::Max(1, InNumMipMaps);
+	SequenceInfo.NumTiles.X = FMath::Max(1, InNumTiles.X);
+	SequenceInfo.NumTiles.Y = FMath::Max(1, InNumTiles.Y);
 }
 
-int32 FImgMediaMipMapInfo::GetDesiredMipLevel(FImgMediaTileSelection& OutTileSelection)
+TMap<int32, FImgMediaTileSelection> FImgMediaMipMapInfo::GetVisibleTiles()
 {
 	// This is called from the loader one thread at a time as the call is guarded by a critical section.
 	// So no need for thread safety here with regards to this function.
@@ -144,179 +489,32 @@ int32 FImgMediaMipMapInfo::GetDesiredMipLevel(FImgMediaTileSelection& OutTileSel
 	// accessing things that are modified by code external to this function.
 	
 	// Do we need to update the cache?
-	if (bIsCachedMipLevelValid == false)
+	if (bIsCacheValid == false)
 	{
 		UpdateMipLevelCache();
 	}
 
-	// Get desired mip level from the cache.
-	int32 MipLevel = CachedMipLevel;
-	OutTileSelection = CachedTileSelection;
-
-	return MipLevel;
+	return CachedVisibleTiles;
 }
 
 void FImgMediaMipMapInfo::UpdateMipLevelCache()
 {
 	SCOPE_CYCLE_COUNTER(STAT_ImgMedia_MipMapUpdateCache);
 
-	CachedMipLevel = 0;
-	CachedTileSelection.SetAllVisible();
-	bool bHasTiles = (NumTilesX > 1) || (NumTilesY > 1);
-
-	// Loop over all objects.
-	float ClosestDistToCamera = TNumericLimits<float>::Max();
 	{
 		FScopeLock Lock(&InfoCriticalSection);
-		for (int Index = 0; Index < Objects.Num(); ++Index)
-		{
-			FImgMediaMipMapObjectInfo* Info = Objects[Index];
-			AActor* Object = Info->Object.Get();
-			if (Object != nullptr)
-			{
-				FVector ObjectLocation = Object->GetActorLocation();
-				
-				const FTransform ObjectTransform = Object->GetTransform();
-				FVector ObjectScale = Object->GetActorScale();
-				
-				// Set up for tiles.
-				float TileWidthInWorldSpace = 0.0f;
-				float TileHeightInWorldSpace = 0.0f;
-				float TileRadiusInWorldSpace = 0.0f;
-				FVector NextTileXVector;
-				FVector NextTileYVector;
-				if (bHasTiles)
-				{
-					// Get size of tile in world space.
-					TileWidthInWorldSpace = ObjectScale.X * Info->Width / ((float)NumTilesX);
-					TileHeightInWorldSpace = ObjectScale.Y * Info->Height / ((float)NumTilesY);
-					TileRadiusInWorldSpace = FMath::Max(TileWidthInWorldSpace, TileHeightInWorldSpace) * 0.5f;
-
-					// Get vectors to go from one tile to the next in world space.
-					FVector TilePosInObjectSpace(0.0f, TileWidthInWorldSpace, 0.0f);
-					FVector TilePos = ObjectTransform.TransformPositionNoScale(TilePosInObjectSpace);
-					NextTileXVector = TilePos - ObjectLocation;
-					TilePosInObjectSpace.Set(0.0f, 0.0f, TileHeightInWorldSpace);
-					TilePos = ObjectTransform.TransformPositionNoScale(TilePosInObjectSpace);
-					NextTileYVector = TilePos - ObjectLocation;
-				}
-
-				for (const FImgMediaMipMapCameraInfo& CameraInfo : CameraInfos)
-				{
-					// Mips calculation.
-					// Get closest distance to camera.
-					float DistToCamera = GetObjectDistToCamera(CameraInfo.Location, ObjectLocation);
-					float AdjustedDistToCamera = DistToCamera;
-
-					// Account for object size and possible rotation.
-					// Closest case is the object edge is pointing towards us.
-					AdjustedDistToCamera = AdjustedDistToCamera - Info->Width * 0.5f;
-
-					// If we are using our calculated distances,
-					// then account for various settings.
-					if (bIsMipLevel0DistanceSetManually == false)
-					{
-						AdjustedDistToCamera *= CameraInfo.DistAdjust;
-						if (CameraInfo.ScreenSize == 0.0f)
-						{
-							AdjustedDistToCamera *= ViewportDistAdjust;
-						}
-						AdjustedDistToCamera *= Info->DistAdjust;
-					}
 
-					if (AdjustedDistToCamera < ClosestDistToCamera)
-					{
-						ClosestDistToCamera = AdjustedDistToCamera;
-					}
+		CachedVisibleTiles.Reset();
 
-					// See which tiles are visible.
-					if (bHasTiles)
-					{
-						CalculateTileVisibility(CameraInfo, ObjectLocation, NextTileXVector,
-							NextTileYVector, TileRadiusInWorldSpace);
-					}
-				}
-			}
+		// Loop over all objects.
+		for (FImgMediaMipMapObjectInfo* ObjectInfo : Objects)
+		{
+			ObjectInfo->CalculateVisibleTiles(CameraInfos, SequenceInfo, CachedVisibleTiles);
 		}
 	}
-
-	// Did we get a distance?
-	if (ClosestDistToCamera != TNumericLimits<float>::Max())
-	{
-		// Find the best mip level.
-		CachedMipLevel = GetMipLevelForDistance(ClosestDistToCamera, MipLevelDistances);
-	}
 	
 	// Mark cache as valid.
-	bIsCachedMipLevelValid = true;
-}
-
-void FImgMediaMipMapInfo::CalculateTileVisibility(const FImgMediaMipMapCameraInfo& CameraInfo,
-	const FVector& ObjectLocation, const FVector& NextTileXVector,
-	const FVector& NextTileYVector, float TileRadiusInWorldSpace)
-{
-	// Get frustum.
-	FConvexVolume ViewFrustum;
-	GetViewFrustumBounds(ViewFrustum, CameraInfo.ViewMatrix, false);
-
-	// This is the middle of the object.
-	FVector TileStartPos = ObjectLocation;
-
-	// Loop over Y.
-	for (int32 TileY = 0; TileY < NumTilesY; ++TileY)
-	{
-		// Y needs to be flipped so that 0 is at the top.
-		// - NumTiles/2 so we start at the corner.
-		// + 0.5 so the position is the middle of the tile.
-		float TileMultY = (NumTilesY - TileY - 1) - ((float)NumTilesY) / 2.0f + 0.5f;
-
-		// Loop over X.
-		for (int32 TileX = 0; TileX < NumTilesX; ++TileX)
-		{
-			float TileMultX = TileX - ((float)NumTilesX) / 2.0f + 0.5f;
-			FVector TilePos = TileStartPos + NextTileXVector * TileMultX +
-				NextTileYVector * TileMultY;
-
-			// Is this tile visible?
-			if (ViewFrustum.IntersectSphere(TilePos, TileRadiusInWorldSpace))
-			{
-				// Is the cache initialized to all visible?
-				if (CachedTileSelection.BottomRightX != 0xffff)
-				{
-					// Nope.
-					// Add this tile to the selection.
-					CachedTileSelection.TopLeftX = FMath::Min(
-						CachedTileSelection.TopLeftX, (uint16)TileX);
-					CachedTileSelection.TopLeftY = FMath::Min(
-						CachedTileSelection.TopLeftY, (uint16)TileY);
-					CachedTileSelection.BottomRightX = FMath::Max(
-						CachedTileSelection.BottomRightX, (uint16)(TileX + 1));
-					CachedTileSelection.BottomRightY = FMath::Max(
-						CachedTileSelection.BottomRightY, (uint16)(TileY + 1));
-				}
-				else
-				{
-					// Yes.
-					// Set tile selection to this tile.
-					CachedTileSelection.TopLeftX = TileX;
-					CachedTileSelection.TopLeftY = TileY;
-					CachedTileSelection.BottomRightX = TileX + 1;
-					CachedTileSelection.BottomRightY = TileY + 1;
-				}
-				
-				// Enable this to draw a sphere where each tile is.
-#if false
-#if WITH_EDITOR
-				Async(EAsyncExecution::TaskGraphMainThread, [TilePos, TileRadiusInWorldSpace]()
-				{
-					UWorld* World = GEditor->GetEditorWorldContext().World();
-					DrawDebugSphere(World, TilePos, TileRadiusInWorldSpace, 6, FColor::Red);
-				});
-#endif // WITH_EDITOR
-#endif // false
-			}
-		}
-	}
+	bIsCacheValid = true;
 }
 
 void FImgMediaMipMapInfo::Tick(float DeltaTime)
@@ -325,95 +523,34 @@ void FImgMediaMipMapInfo::Tick(float DeltaTime)
 
 	// Get global info.
 	FImgMediaMipMapInfoManager& InfoManager = FImgMediaMipMapInfoManager::Get();
-	ViewportDistAdjust = InfoManager.GetViewportInfo();
 	CameraInfos = InfoManager.GetCameraInfo();
 
 	// Let the cache update this frame.
-	bIsCachedMipLevelValid = false;
+	bIsCacheValid = false;
 
 	// Display debug?
 	if (InfoManager.IsDebugEnabled())
 	{
 		if (GEngine != nullptr)
 		{
-			GEngine->AddOnScreenDebugMessage(-1, 0.0f, FColor::Yellow, *FString::Printf(TEXT("%s Mip Level: %d"), *SequenceName.ToString(), CachedMipLevel));
-		}
-	}
-}
-
-void FImgMediaMipMapInfo::UpdateMipLevelDistances()
-{
-	if ((MipLevelDistances.Num() > 0) && (MipLevel0Distance >= 0.0f))
-	{
-		MipLevelDistances[0] = MipLevel0Distance;
-
-		// Objects will appear half the size when they are twice the distance away.
-		// So we can approximate the mip level distance by each level being twice the distance
-		// of the previous level, since the mip level is half the size of the previous level.
-		for (int Level = 1; Level < MipLevelDistances.Num(); ++Level)
-		{
-			MipLevelDistances[Level] = MipLevelDistances[Level - 1] * 2.0f;
-		}
-
-		for (int Level = 0; Level < MipLevelDistances.Num(); ++Level)
-			UE_LOG(LogImgMedia, VeryVerbose, TEXT("FImgMediaMipMapInfo MipLevel:%d Distance:%f"), Level, MipLevelDistances[Level]);
-	}
-}
-
-float FImgMediaMipMapInfo::GetObjectDistToCamera(const FVector& InCameraLocation, const FVector& InObjectLocation)
-{
-	float DistToCamera = FVector::Dist(InCameraLocation, InObjectLocation);
-
-	return DistToCamera;
-}
+			FIntVector2 MipLevelRange = FIntVector2(TNumericLimits<int32>::Max(), 0);
+			int32 NumVisibleTiles = 0;
+			for (const auto& MipTiles : CachedVisibleTiles)
+			{
+				MipLevelRange[0] = FMath::Min(MipLevelRange[0], MipTiles.Key);
+				MipLevelRange[1] = FMath::Max(MipLevelRange[1], MipTiles.Key);
 
-int FImgMediaMipMapInfo::GetMipLevelForDistance(float InDistance, const TArray<float>& InMipLevelDistances)
-{
-	int MipLevel = 0;
+				const FImgMediaTileSelection& TileSelection = MipTiles.Value;
+				NumVisibleTiles += (TileSelection.BottomRightX - TileSelection.TopLeftX) * (TileSelection.BottomRightY - TileSelection.TopLeftY);
+			}
 
-	// Find the best mip level.
-	for (MipLevel = 0; MipLevel < InMipLevelDistances.Num() - 1; ++MipLevel)
-	{
-		// If we are between levels 2 and 3, then return 2.
-		if (InDistance <= InMipLevelDistances[MipLevel + 1])
-		{
-			break;
+			if (MipLevelRange[0] <= MipLevelRange[1])
+			{
+				GEngine->AddOnScreenDebugMessage(-1, 0.0f, FColor::Yellow, *FString::Printf(TEXT("%s Mip Level: [%d, %d]"), *SequenceInfo.Name.ToString(), MipLevelRange[0], MipLevelRange[1]));
+				GEngine->AddOnScreenDebugMessage(-1, 0.0f, FColor::Yellow, *FString::Printf(TEXT("%s Num Tiles: %d"), *SequenceInfo.Name.ToString(), NumVisibleTiles));
+			}
 		}
 	}
-	
-	return MipLevel;
 }
 
-void FImgMediaMipMapInfo::GetObjectSize(const AActor* InActor, float& Width, float& Height)
-{
-	// Get box extent of actor to calculate width from.
-	Width = -1.0f;
-	if (InActor != nullptr)
-	{
-		FVector BoundsOrigin;
-		FVector BoxExtent;
-		InActor->GetActorBounds(false, BoundsOrigin, BoxExtent);
-		Width = BoxExtent.Z * 2.0f;
-		Height = BoxExtent.Y * 2.0f;
-
-		// We apply the object scsale in UpdateMipLevelCache, so take it out here.
-		FVector ObjectScale = InActor->GetActorScale();
-		if (FMath::IsNearlyZero(ObjectScale.X) == false)
-		{
-			Width = Width / ObjectScale.X;
-		}
-		if (FMath::IsNearlyZero(ObjectScale.Y) == false)
-		{
-			Height = Height / ObjectScale.Y;
-		}
-	}
-
-	// Did we get anything?
-	if (Width < 0.0f)
-	{
-		UE_LOG(LogImgMedia, Error, TEXT("FImgMediaMipMapInfo could not get size of %s."), InActor != nullptr ? *InActor->GetName() : TEXT("<nullptr>"));
-		Width = 0.0f;
-		Height = 0.0f;
-	}
-}
 
diff --git a/Engine/Plugins/Media/ImgMedia/Source/ImgMedia/Private/Assets/ImgMediaMipMapInfoManager.cpp b/Engine/Plugins/Media/ImgMedia/Source/ImgMedia/Private/Assets/ImgMediaMipMapInfoManager.cpp
--- a/Engine/Plugins/Media/ImgMedia/Source/ImgMedia/Private/Assets/ImgMediaMipMapInfoManager.cpp
+++ b/Engine/Plugins/Media/ImgMedia/Source/ImgMedia/Private/Assets/ImgMediaMipMapInfoManager.cpp
@@ -30,34 +30,7 @@ static TAutoConsoleVariable<bool> CVarImgMediaMipMapDebugEnable(
 	ECVF_Default);
 
 FImgMediaMipMapInfoManager::FImgMediaMipMapInfoManager()
-	: RefNearPlaneWidth(0.0f)
-	, MipLevel0Distance(0.0f)
-	, ViewportDistAdjust(0.0f)
 {
-	// Get size of the near plane in physical space.
-	RefNearPlaneWidth = CalculateNearPlaneSize(RefCameraFOV);
-
-	// Get size of the near plane in physical space.
-	float NearPlaneWidth = RefNearPlaneWidth;
-
-	// Calculate how big the object is at the near plane (i.e. the screen).
-	float ObjectWidth = RefObjectWidth;
-	float ObjectDistance = RefObjectDistance;
-	float NearPlaneDistance = GetRefNearPlaneDistance();
-	float ObjectWidthAtNearPlane = (ObjectWidth / ObjectDistance) * NearPlaneDistance;
-
-	float FrameBufferWidth = RefFrameBufferWidth;
-	float ObjectTextureWidth = RefObjectTextureWidth;
-
-	// Calculate how big in frame buffer pixels the object is on the screen.
-	float ObjectPixelWidth = (ObjectWidthAtNearPlane / NearPlaneWidth) * FrameBufferWidth;
-	// Calculate how many texels are displayed per frame buffer pixel.
-	float ObjectTexelsPerPixel = ObjectTextureWidth / ObjectPixelWidth;
-	// How far the object should be so that one texel covers one frame buffer pixel.
-	float ObjectDistanceForOneTexelPerPixel = ObjectDistance / ObjectTexelsPerPixel;
-
-	MipLevel0Distance = ObjectDistanceForOneTexelPerPixel;
-
 	CVarImgMediaMipMapDebugEnable.AsVariable()->SetOnChangedCallback(FConsoleVariableDelegate::CreateRaw(this, &FImgMediaMipMapInfoManager::OnCVarMipMapDebugEnable));
 	
 	// Make sure info is set up.
@@ -132,26 +105,9 @@ void FImgMediaMipMapInfoManager::GetMediaTexturesFromPlayer(TArray<UMediaTexture
 
 void FImgMediaMipMapInfoManager::Tick(float DeltaTime)
 {
-	UpdateViewportInfo();
 	UpdateCameraInfo();
 }
 
-void FImgMediaMipMapInfoManager::UpdateViewportInfo()
-{
-	// Get main viewport.
-	UEngine* Engine = GEngine;
-	if (Engine != nullptr)
-	{
-		UGameViewportClient* Viewport = Engine->GameViewport;
-		if (Viewport != nullptr)
-		{
-			FVector2D ViewportSize;
-			Viewport->GetViewportSize(ViewportSize);
-			ViewportDistAdjust = RefFrameBufferWidth / ViewportSize.X;
-		}
-	}
-}
-
 void FImgMediaMipMapInfoManager::UpdateCameraInfo()
 {
 	CameraInfos.Empty();
@@ -160,30 +116,34 @@ void FImgMediaMipMapInfoManager::UpdateCameraInfo()
 	UGameViewportClient* GameViewportClient = GEngine->GameViewport;
 	if (GameViewportClient != nullptr)
 	{
-		UWorld* World = GameViewportClient->GetWorld();
-		if (World != nullptr)
+		FViewport* Viewport = GameViewportClient->Viewport;
+		if (Viewport != nullptr)
 		{
-			for (FConstPlayerControllerIterator Iterator = World->GetPlayerControllerIterator(); Iterator; ++Iterator)
+			UWorld* World = GameViewportClient->GetWorld();
+			if (World != nullptr)
 			{
-				APlayerController* PlayerController = Iterator->Get();
-				ULocalPlayer* LocalPlayer = PlayerController->GetLocalPlayer();
-				if (LocalPlayer != nullptr)
+				for (FConstPlayerControllerIterator Iterator = World->GetPlayerControllerIterator(); Iterator; ++Iterator)
 				{
-					// Get the scene view.
-					FSceneViewFamilyContext ViewFamily(FSceneViewFamily::ConstructionValues(
-						LocalPlayer->ViewportClient->Viewport,
-						World->Scene,
-						LocalPlayer->ViewportClient->EngineShowFlags)
-						.SetRealtimeUpdate(true));
-					FVector ViewLocation;
-					FRotator ViewRotation;
-					FSceneView* SceneView = LocalPlayer->CalcSceneView(&ViewFamily, ViewLocation, ViewRotation, LocalPlayer->ViewportClient->Viewport);
-
-					// Did we get a view?
-					if (SceneView != nullptr)
+					APlayerController* PlayerController = Iterator->Get();
+					ULocalPlayer* LocalPlayer = PlayerController->GetLocalPlayer();
+					if (LocalPlayer != nullptr)
 					{
-						// Add camera info for this scene view.
-						AddCameraInfo(SceneView);
+						// Get the scene view.
+						FSceneViewFamilyContext ViewFamily(FSceneViewFamily::ConstructionValues(
+							LocalPlayer->ViewportClient->Viewport,
+							World->Scene,
+							LocalPlayer->ViewportClient->EngineShowFlags)
+							.SetRealtimeUpdate(true));
+						FVector ViewLocation;
+						FRotator ViewRotation;
+						FSceneView* SceneView = LocalPlayer->CalcSceneView(&ViewFamily, ViewLocation, ViewRotation, LocalPlayer->ViewportClient->Viewport);
+
+						// Did we get a view?
+						if (SceneView != nullptr)
+						{
+							// Add camera info for this scene view.
+							AddCameraInfo(Viewport, SceneView);
+						}
 					}
 				}
 			}
@@ -217,7 +177,7 @@ void FImgMediaMipMapInfoManager::UpdateCameraInfo()
 						FSceneView* SceneView = ViewportClient->CalcSceneView(&ViewFamily);
 						
 						// Add camera info for this scene view.
-						AddCameraInfo(SceneView);
+						AddCameraInfo(Viewport, SceneView);
 					}
 				}
 			}
@@ -226,27 +186,16 @@ void FImgMediaMipMapInfoManager::UpdateCameraInfo()
 #endif
 }
 
-void FImgMediaMipMapInfoManager::AddCameraInfo(FSceneView* SceneView)
+void FImgMediaMipMapInfoManager::AddCameraInfo(FViewport* Viewport, FSceneView* SceneView)
 {
-	float FOV = SceneView->FOV;
-	float ViewRectWidth = SceneView->UnscaledViewRect.Width();
-	float ViewRectHeight = SceneView->UnscaledViewRect.Height();
-		
-	const FVector& ViewLocation = SceneView->ViewMatrices.GetViewOrigin();
-	
-	float NearPlaneSize = CalculateNearPlaneSize(FOV);
-	float DistAdjust = NearPlaneSize / RefNearPlaneWidth;
-	DistAdjust *= RefFrameBufferWidth / ViewRectWidth;
-
-	CameraInfos.Emplace(ViewLocation, SceneView->ViewMatrices.GetViewProjectionMatrix(), ViewRectWidth, DistAdjust);
-}
-
-float FImgMediaMipMapInfoManager::CalculateNearPlaneSize(float InFOV) const
-{
-	float NearPlaneDistance = RefNearPlaneDistance;
-	float NearPlaneWidth = FMath::Tan(FMath::DegreesToRadians(InFOV) * 0.5f) * NearPlaneDistance * 2.0f;
-
-	return NearPlaneWidth;
+	FImgMediaMipMapCameraInfo Info;
+	Info.Location = SceneView->ViewMatrices.GetViewOrigin();
+	Info.ViewMatrix = SceneView->ViewMatrices.GetViewMatrix();
+	Info.ViewProjectionMatrix = SceneView->ViewMatrices.GetViewProjectionMatrix();
+	Info.ViewportRect = FIntRect(FIntPoint::ZeroValue, Viewport->GetSizeXY());
+	Info.MaterialTextureMipBias = SceneView->MaterialTextureMipBias;
+
+	CameraInfos.Emplace(MoveTemp(Info));
 }
 
 void FImgMediaMipMapInfoManager::OnCVarMipMapDebugEnable(IConsoleVariable* Var)
diff --git a/Engine/Plugins/Media/ImgMedia/Source/ImgMedia/Private/Assets/ImgMediaMipMapInfoManager.h b/Engine/Plugins/Media/ImgMedia/Source/ImgMedia/Private/Assets/ImgMediaMipMapInfoManager.h
--- a/Engine/Plugins/Media/ImgMedia/Source/ImgMedia/Private/Assets/ImgMediaMipMapInfoManager.h
+++ b/Engine/Plugins/Media/ImgMedia/Source/ImgMedia/Private/Assets/ImgMediaMipMapInfoManager.h
@@ -56,13 +56,6 @@ public:
 	 */
 	void GetMediaTexturesFromPlayer(TArray<UMediaTexture*>& OutMediaTextures, IMediaPlayer* Player);
 
-	/**
-	 * Get information on the main viewport.
-	 *
-	 * @return Distance adjustment for this viewport compared to the reference viewport.
-	 */
-	float GetViewportInfo() { return ViewportDistAdjust; }
-
 	/**
 	 * Get information on all appropriate cameras.
 	 *
@@ -70,36 +63,6 @@ public:
 	 */
 	const TArray<FImgMediaMipMapCameraInfo>& GetCameraInfo() { return CameraInfos; };
 
-	/**
-	 * Get the width (in physical space) of the reference near plane used in mipmap calculations.
-	 *
-	 * @return Width
-	 */
-	float GetRefNearPlaneWidth() const { return RefNearPlaneWidth; }
-
-	/**
-	 * Get the distance of the reference near plane from the camera used in mipmap calculations.
-	 *
-	 * @return Distance.
-	 */
-	float GetRefNearPlaneDistance() const { return RefNearPlaneDistance; }
-
-	/**
-	 * Get the size of the reference object.
-	 */
-	float GetRefObjectWidth() const { return RefObjectWidth; }
-
-	/**
-	 * Get the size of the texture for the reference object.
-	 */
-	float GetRefObjectTextureWidth() const { return RefObjectTextureWidth; }
-
-	/**
-	 * Get the distance of mip level 0 for the reference object.
-	 *
-	 * @return Distance.
-	 */
-	float GetMipLevel0Distance() const { return MipLevel0Distance; }
 
 	/**
 	 * See if we want to output some debug.
@@ -121,34 +84,9 @@ protected:
 	/** Array of info on each camera used for mipmap calculations. */
 	TArray<FImgMediaMipMapCameraInfo> CameraInfos;
 
-	/** Reference camera FOV used in mipmap calculations. */
-	const float RefCameraFOV = 90.0f;
-	/** Size of reference near plane (in physical space) used to calculate MipLevelDistance. */
-	float RefNearPlaneWidth;
-	/** Distance of reference near plan used in mipmap calculations. */
-	const float RefNearPlaneDistance = 10.0f;
-	/** Size of reference object used to calculate MipLevelDistances. */
-	const float RefObjectWidth = 100.0f;
-	/** Distance of reference object used to calculate MipLevelDistances. */
-	const float RefObjectDistance = 100.0f;
-	/** Texture size of reference object used to calculate MipLevelDistances. */
-	const float RefObjectTextureWidth = 1024.0f;
-	/** Size of reference frame buffer used to calculate MipLevelDistances. */
-	const float RefFrameBufferWidth = 1920.0f;
-
-	/** Distance for reference object for mip level 0. */
-	float MipLevel0Distance;
-	/** Adjustment needed for main viewport. */
-	float ViewportDistAdjust;
-
 	/** Do we want to output debug information? */
 	bool bIsDebugEnabled;
 
-	/**
-	 * Updates our info on the main viewport.
-	 */
-	void UpdateViewportInfo();
-
 	/**
 	 * Updates our info on the cameras based on current state.
 	 */
@@ -157,16 +95,10 @@ protected:
 	/**
 	 * Add a SceneView to our camera info.
 	 * 
+	 * @param Viewport		The camera's viewport size rectangle.
 	 * @param SceneView		The scene view we want to add.
 	 */
-	void AddCameraInfo(FSceneView* SceneView);
-
-	/** Calculate the size of the near plane in physical space.
-	 *
-	 * @param InFOV FOV of camera in degrees.
-	 * @return Size of near plane.
-	 */
-	float CalculateNearPlaneSize(float InFOV) const;
+	void AddCameraInfo(FViewport* Viewport, FSceneView* SceneView);
 
 	/**
 	 * Called when the cvar changes.
diff --git a/Engine/Plugins/Media/ImgMedia/Source/ImgMedia/Private/Assets/ImgMediaSource.cpp b/Engine/Plugins/Media/ImgMedia/Source/ImgMedia/Private/Assets/ImgMediaSource.cpp
--- a/Engine/Plugins/Media/ImgMedia/Source/ImgMedia/Private/Assets/ImgMediaSource.cpp
+++ b/Engine/Plugins/Media/ImgMedia/Source/ImgMedia/Private/Assets/ImgMediaSource.cpp
@@ -83,12 +83,6 @@ void UImgMediaSource::RemoveTargetObject(AActor* InActor)
 }
 
 
-void UImgMediaSource::SetMipLevelDistance(float Distance)
-{
-	MipMapInfo->SetMipLevelDistance(Distance);
-}
-
-
 /* IMediaOptions interface
  *****************************************************************************/
 
diff --git a/Engine/Plugins/Media/ImgMedia/Source/ImgMedia/Private/Loader/ImgMediaLoader.cpp b/Engine/Plugins/Media/ImgMedia/Source/ImgMedia/Private/Loader/ImgMediaLoader.cpp
--- a/Engine/Plugins/Media/ImgMedia/Source/ImgMedia/Private/Loader/ImgMediaLoader.cpp
+++ b/Engine/Plugins/Media/ImgMedia/Source/ImgMedia/Private/Loader/ImgMediaLoader.cpp
@@ -51,24 +51,6 @@ namespace ImgMediaLoader
 		}
 	}
 	
-
-	FMediaTextureTilingDescription GetTilingDescription(const FImgMediaLoader* Loader)
-	{
-		FMediaTextureTilingDescription TilingDesc;
-		
-		if (Loader->IsTiled())
-		{
-			const int32 NumX = Loader->GetNumTilesX();
-			const int32 NumY = Loader->GetNumTilesY();
-
-			TilingDesc.TileBorderSize = 0; // todo: get from loader once accessible from metadata
-			TilingDesc.TileSize = FIntPoint(Loader->GetSequenceDim().X / NumX, Loader->GetSequenceDim().Y / NumY);
-			TilingDesc.TileNum = FIntPoint(NumX, NumY);
-		}
-		
-		return TilingDesc;
-	}
-
 	// Check if the existing tiles contain all of the requested ones. (Is existing a superset of requested?)
 	bool ContainsMipTiles(const TMap<int32, FImgMediaTileSelection>& ExistingTiles, const TMap<int32, FImgMediaTileSelection>& RequestedTiles)
 	{
@@ -95,25 +77,6 @@ namespace ImgMediaLoader
 		// Requested tiles already exist, or the request was empty.
 		return true;
 	}
-
-	// Reduce the selected tiles region for higher mip levels, always ensuring that all tiles at the base level are included.
-	FImgMediaTileSelection DownscaleTileSelection(const FImgMediaTileSelection& InTopLevelSelection, int32 InMipLevel)
-	{
-		if (InMipLevel <= 0)
-		{
-			return InTopLevelSelection;
-		}
-		
-		const int32 MipLevelDiv = 1 << InMipLevel;
-
-		FImgMediaTileSelection OutSelection;
-		OutSelection.TopLeftX = (uint16)(InTopLevelSelection.TopLeftX / MipLevelDiv);
-		OutSelection.TopLeftY = (uint16)(InTopLevelSelection.TopLeftY / MipLevelDiv);
-		OutSelection.BottomRightX = (uint16)FMath::CeilToInt(float(InTopLevelSelection.BottomRightX) / MipLevelDiv);
-		OutSelection.BottomRightY = (uint16)FMath::CeilToInt(float(InTopLevelSelection.BottomRightY) / MipLevelDiv);
-
-		return OutSelection;
-	}
 }
 
 /* FImgMediaLoader structors
@@ -131,8 +94,7 @@ FImgMediaLoader::FImgMediaLoader(const TSharedRef<FImgMediaScheduler, ESPMode::T
 	, bFillGapsInSequence(bInFillGapsInSequence)
 	, bReadVirtualTextureTiles(bInReadVirtualTextureTiles)
 	, bIsTiled(false)
-	, NumTilesX(0)
-	, NumTilesY(0)
+	, TilingDescription()
 	, NumLoadAhead(0)
 	, NumLoadBehind(0)
 	, Scheduler(InScheduler)
@@ -242,7 +204,7 @@ TSharedPtr<FImgMediaTextureSample, ESPMode::ThreadSafe> FImgMediaLoader::GetFram
 
 	ImgMediaLoader::CheckAndUpdateImgDimensions(SequenceDim, Frame->Get()->Info.Dim);
 
-	if (!Sample->Initialize(*Frame->Get(), SequenceDim, FMediaTimeStamp(FrameStartTime, 0), NextStartTime - FrameStartTime, GetNumMipLevels(), ImgMediaLoader::GetTilingDescription(this)))
+	if (!Sample->Initialize(*Frame->Get(), SequenceDim, FMediaTimeStamp(FrameStartTime, 0), NextStartTime - FrameStartTime, GetNumMipLevels(), TilingDescription))
 	{
 		return nullptr;
 	}
@@ -543,7 +505,7 @@ IMediaSamples::EFetchBestSampleResult FImgMediaLoader::FetchBestVideoSampleForTi
 
 					ImgMediaLoader::CheckAndUpdateImgDimensions(SequenceDim, Frame->Get()->Info.Dim);
 
-					if (Sample->Initialize(*Frame->Get(), SequenceDim, FMediaTimeStamp(FrameNumberToTime(MaxIdx), QueuedSampleFetch.CurrentSequenceIndex), FTimespan::FromSeconds(Duration), GetNumMipLevels(), ImgMediaLoader::GetTilingDescription(this)))
+					if (Sample->Initialize(*Frame->Get(), SequenceDim, FMediaTimeStamp(FrameNumberToTime(MaxIdx), QueuedSampleFetch.CurrentSequenceIndex), FTimespan::FromSeconds(Duration), GetNumMipLevels(), TilingDescription))
 					{
 						OutSample = Sample;
 						CSV_EVENT(ImgMedia, TEXT("LoaderFetchHit %d %d-%d"), MaxIdx, StartIndex, EndIndex);
@@ -885,14 +847,12 @@ void FImgMediaLoader::LoadSequence(const FString& SequencePath, const FFrameRate
 		MinBandwidthForMipLevel0 /= FrameTime;
 	}
 
-#if 0
-	// todo: enable once both the CPU & GPU reader use the same tile count logic.
 	if (FirstFrameInfo.bHasTiles)
 	{
-		NumTilesX = FirstFrameInfo.NumTiles.X;
-		NumTilesY = FirstFrameInfo.NumTiles.Y;
+		TilingDescription.TileBorderSize = FirstFrameInfo.TileBorder;
+		TilingDescription.TileNum = FirstFrameInfo.NumTiles;
+		TilingDescription.TileSize = FirstFrameInfo.TileDimensions;
 	}
-#endif
 
 	// If we have no mips or tiles, then get rid of our MipMapInfoObject.
 	// Otherwise, set it up.
@@ -904,8 +864,7 @@ void FImgMediaLoader::LoadSequence(const FString& SequencePath, const FFrameRate
 		}
 		else
 		{
-			MipMapInfo->SetTextureInfo(SequenceName, GetNumMipLevels(), NumTilesX, NumTilesY,
-				SequenceDim);
+			MipMapInfo->SetTextureInfo(SequenceName, GetNumMipLevels(), TilingDescription.TileNum, SequenceDim);
 			if (GetNumMipLevels() > 1)
 			{
 				UncompressedSize = (UncompressedSize * 4) / 3;
@@ -1099,8 +1058,8 @@ void FImgMediaLoader::FindMips(const FString& SequencePath)
 
 void FImgMediaLoader::GetTileInformation()
 {
-	NumTilesX = 1;
-	NumTilesY = 1;
+	TilingDescription.TileNum.X = 1;
+	TilingDescription.TileNum.Y = 1;
 
 	// Do we have tiles?
 	if (bIsTiled)
@@ -1123,7 +1082,7 @@ void FImgMediaLoader::GetTileInformation()
 				int32 YLength = Path.Len() - (YIndex + 1);
 				FString NumberString = Path.Right(YLength);
 				// We count from 0, so add 1 to get the number of tiles.
-				NumTilesY = FMath::Max(NumTilesY, FCString::Atoi(*NumberString) + 1);
+				TilingDescription.TileNum.Y = FMath::Max(TilingDescription.TileNum.Y, FCString::Atoi(*NumberString) + 1);
 
 				// Find the last x.
 				int32 XIndex = 0;
@@ -1133,7 +1092,7 @@ void FImgMediaLoader::GetTileInformation()
 					XIndex++;
 					int32 XLength = Path.Len() - XIndex - YLength - 2;
 					NumberString = Path.Mid(XIndex, XLength);
-					NumTilesX = FMath::Max(NumTilesX, FCString::Atoi(*NumberString) + 1);
+					TilingDescription.TileNum.X = FMath::Max(TilingDescription.TileNum.X, FCString::Atoi(*NumberString) + 1);
 				}
 			}
 		}
@@ -1408,19 +1367,20 @@ void FImgMediaLoader::GetDesiredMipTiles(int32 FrameIndex, TMap<int32, FImgMedia
 			}
 		}
 	}
+	else if(MipMapInfo.IsValid() && MipMapInfo->HasObjects())
+	{
+		OutMipsAndTiles = MipMapInfo->GetVisibleTiles();
+	}
 	else
 	{
-		FImgMediaTileSelection BaseTileSelection;
-		BaseTileSelection.SetAllVisible();
-
-		const int32 MipLevel = MipMapInfo.IsValid() ? MipMapInfo->GetDesiredMipLevel(BaseTileSelection) : 0;
-
-		OutMipsAndTiles.Emplace(MipLevel, ImgMediaLoader::DownscaleTileSelection(BaseTileSelection, MipLevel));
-
-		// Reproducing the previous mip-threshold logic by adding all higher mip levels.
-		for (int32 HigherMipLevel = MipLevel + 1; HigherMipLevel < GetNumMipLevels(); ++HigherMipLevel)
+		//Fallback case where we activate all mips and tiles.
+		for (int32 MipLevel = 0; MipLevel < GetNumMipLevels(); ++MipLevel)
 		{
-			OutMipsAndTiles.Emplace(HigherMipLevel, ImgMediaLoader::DownscaleTileSelection(BaseTileSelection, HigherMipLevel));
+			FImgMediaTileSelection& Selection = OutMipsAndTiles.FindOrAdd(MipLevel);
+
+			int32 NumTilesX = FMath::Max(1, GetNumTilesX() >> MipLevel);
+			int32 NumTilesY = FMath::Max(1, GetNumTilesY() >> MipLevel);
+			Selection.SetVisibleRegion(0u, 0u, IntCastChecked<uint16>(NumTilesX), IntCastChecked<uint16>(NumTilesY));
 		}
 	}
 }
diff --git a/Engine/Plugins/Media/ImgMedia/Source/ImgMedia/Private/Loader/ImgMediaLoader.h b/Engine/Plugins/Media/ImgMedia/Source/ImgMedia/Private/Loader/ImgMediaLoader.h
--- a/Engine/Plugins/Media/ImgMedia/Source/ImgMedia/Private/Loader/ImgMediaLoader.h
+++ b/Engine/Plugins/Media/ImgMedia/Source/ImgMedia/Private/Loader/ImgMediaLoader.h
@@ -6,6 +6,7 @@
 #include "Containers/LruCache.h"
 #include "Containers/Queue.h"
 #include "IMediaSamples.h"
+#include "IMediaTextureSample.h"
 #include "IMediaView.h"
 #include "ImgMediaMipMapInfo.h"
 #include "Misc/FrameRate.h"
@@ -241,7 +242,7 @@ public:
 	 */
 	int32 GetNumTilesX() const
 	{
-		return NumTilesX;
+		return TilingDescription.TileNum.X;
 	}
 
 	/**
@@ -249,12 +250,12 @@ public:
 	 */
 	int32 GetNumTilesY() const
 	{
-		return NumTilesY;
+		return TilingDescription.TileNum.Y;
 	}
 
 	bool IsTiled() const
 	{
-		return (NumTilesX > 1) || (NumTilesY > 1);
+		return (TilingDescription.TileNum.X > 1) || (TilingDescription.TileNum.Y > 1);
 	}
 
 	/**
@@ -496,11 +497,8 @@ private:
 	/** True if this sequence contains tiles. */
 	bool bIsTiled;
 
-	/** Number of tiles in the X direction. */
-	int32 NumTilesX;
-
-	/** Number of tiles in the Y direction. */
-	int32 NumTilesY;
+	/** Tiling description. */
+	FMediaTextureTilingDescription TilingDescription;
 
 	/** The number of frames to load ahead of the play head. */
 	int32 NumLoadAhead;
diff --git a/Engine/Plugins/Media/ImgMedia/Source/ImgMedia/Private/Readers/ExrImgMediaReader.cpp b/Engine/Plugins/Media/ImgMedia/Source/ImgMedia/Private/Readers/ExrImgMediaReader.cpp
--- a/Engine/Plugins/Media/ImgMedia/Source/ImgMedia/Private/Readers/ExrImgMediaReader.cpp
+++ b/Engine/Plugins/Media/ImgMedia/Source/ImgMedia/Private/Readers/ExrImgMediaReader.cpp
@@ -52,23 +52,7 @@ FExrImgMediaReader::~FExrImgMediaReader()
 
 bool FExrImgMediaReader::GetFrameInfo(const FString& ImagePath, FImgMediaFrameInfo& OutInfo)
 {
-	TSharedPtr<FImgMediaLoader, ESPMode::ThreadSafe> Loader = LoaderPtr.Pin();
-	if (Loader.IsValid() == false)
-	{
-		return false;
-	}
-
-	// Get tile info.
-	int32 NumTilesX = Loader->GetNumTilesX();
-	int32 NumTilesY = Loader->GetNumTilesY();
-	if (GetInfo(ImagePath, OutInfo))
-	{
-		OutInfo.Dim.X *= NumTilesX;
-		OutInfo.Dim.Y *= NumTilesY;
-		OutInfo.UncompressedSize *= NumTilesY * NumTilesX;
-		return true;
-	}
-	return false;
+	return GetInfo(ImagePath, OutInfo);
 }
 
 FExrImgMediaReader::EReadResult FExrImgMediaReader::ReadTilesCustom
diff --git a/Engine/Plugins/Media/ImgMedia/Source/ImgMedia/Public/ImgMediaMipMapInfo.h b/Engine/Plugins/Media/ImgMedia/Source/ImgMedia/Public/ImgMediaMipMapInfo.h
--- a/Engine/Plugins/Media/ImgMedia/Source/ImgMedia/Public/ImgMediaMipMapInfo.h
+++ b/Engine/Plugins/Media/ImgMedia/Source/ImgMedia/Public/ImgMediaMipMapInfo.h
@@ -9,49 +9,48 @@
 #include "UObject/WeakObjectPtrTemplates.h"
 
 class AActor;
+class UMeshComponent;
 class UMediaTexture;
 
 /**
- * Holds info on a camera which we can use for mipmap calculations.
+ * Holds size, tiling and mips sequence information.
  */
-struct FImgMediaMipMapCameraInfo
+struct FSequenceInfo
 {
-	FImgMediaMipMapCameraInfo(const FVector& InLocation,
-		const FMatrix& InViewMatrix,
-		float InScreenSize, float InDistAdjust)
-		: Location(InLocation)
-		, ViewMatrix(InViewMatrix)
-		, ScreenSize(InScreenSize)
-		, DistAdjust(InDistAdjust)
+	/** Name of this sequence. */
+	FName Name;
+	/** Pixel dimensions of this sequence. */
+	FIntPoint Dim;
+	/** Number of tiles in the X, Y directions. */
+	FIntPoint NumTiles;
+	/** Number of mip levels. */
+	int32 NumMipLevels;
+	/** 
+	* Check if the sequence has tiles.
+	* 
+	* @return True if the tile counts are above one.
+	*/
+	FORCEINLINE bool IsTiled() const
 	{
+		return (NumTiles.X > 1) || (NumTiles.Y > 1);
 	}
-
-	/** Position of camera. */
-	FVector Location;
-	/** View projection matrix of camera. */
-	FMatrix ViewMatrix;
-
-	/** Size of screen. */
-	float ScreenSize;
-	/** Adjustment needed to mip level distance calculations due to camera FOV, etc. */
-	float DistAdjust;
 };
 
-/** 
- * Describes a single object which is using our img sequence.
+/**
+ * Holds info on a camera which we can use for mipmap calculations.
  */
-struct FImgMediaMipMapObjectInfo
+struct FImgMediaMipMapCameraInfo
 {
-	/** Actor that is using our img sequence. */
-	TWeakObjectPtr<class AActor> Object;
-	/** Width of this object. */
-	float Width;
-	/** Height of this object. */
-	float Height;
-	/** LOD bias for the mipmap level. */
-	float LODBias;
-	/** Multiplier to apply to the distance to account for this object (e.g. its size). */
-	float DistAdjust;
+	/** Position of camera. */
+	FVector Location;
+	/** View matrix of camera. */
+	FMatrix ViewMatrix;
+	/** View projection matrix of camera. */
+	FMatrix ViewProjectionMatrix;
+	/** Active viewport size. */
+	FIntRect ViewportRect;
+	/** View mip bias. */
+	float MaterialTextureMipBias;
 };
 
 /**
@@ -73,7 +72,7 @@ struct FImgMediaTileSelection
 
 	FImgMediaTileSelection()
 	{
-		SetAllVisible();
+		SetAllNotVisible();
 	}
 
 	/**
@@ -98,6 +97,22 @@ struct FImgMediaTileSelection
 		BottomRightY = 0;
 	}
 
+	/**
+	 * Marks visible tile region.
+	 * 
+	 * @param InTopLeftX Top-left coordinate horizontal component.
+	 * @param InTopLeftY Top-left coordinate vertical component.
+	 * @param InBottomRightX Bottom-right coordinate horizontal component.
+	 * @param InBottomRightY Bottom-rightcoordinate vertical component.
+	 */
+	void SetVisibleRegion(uint16 InTopLeftX, uint16 InTopLeftY, uint16 InBottomRightX, uint16 InBottomRightY)
+	{
+		TopLeftX = InTopLeftX;
+		TopLeftY = InTopLeftY;
+		BottomRightX = InBottomRightX;
+		BottomRightY = InBottomRightY;
+	}
+
 	/**
 	* See if this selection is visible.
 	*
@@ -150,6 +165,38 @@ struct FImgMediaTileSelection
 	}
 };
 
+/**
+ * Describes a single object which is using our img sequence.
+ * Base class for various objects such as planes or spheres.
+ */
+class FImgMediaMipMapObjectInfo
+{
+public:
+	FImgMediaMipMapObjectInfo(UMeshComponent* InMeshComponent, float InLODBias = 0.0f);
+	virtual ~FImgMediaMipMapObjectInfo() = default;
+
+	/**
+	 * Calculate visible tiles per mip level.
+	 *
+	 * @param InCameraInfos Active camera information
+	 * @param InSequenceInfo Active image sequence information
+	 * @param VisibleTiles Updated list of visible tiles per mip level
+	 */
+	virtual bool CalculateVisibleTiles(const TArray<FImgMediaMipMapCameraInfo>& InCameraInfos, const FSequenceInfo& InSequenceInfo, TMap<int32, FImgMediaTileSelection>& VisibleTiles) const;
+
+	/**
+	 * Get the registered mesh component.
+	 *
+	 * @return Mesh component
+	 */
+	UMeshComponent* GetMeshComponent() const;
+protected:
+	/** Mesh component onto which the media is displayed. */
+	TWeakObjectPtr<class UMeshComponent> MeshComponent;
+	/** LOD bias for the mipmap level. */
+	float LODBias;
+};
+
 /**
  * Contains information for working with mip maps.
  */
@@ -186,28 +233,16 @@ public:
 	 */
 	void ClearAllObjects();
 
-	/**
-	 * Get our mip level distances.
-	 * @return Distances
-	 */
-	const TArray<float>& GetMipLevelDistances() const { return MipLevelDistances; }
-
-	/**
-	 * Manually set when mip level 0 should appear.
-	 *
-	 * @param Distance Furthest distance from the camera when mip level 0 should be at 100%.
-	 */
-	void SetMipLevelDistance(float Distance);
-
 	/**
 	 * Provide information on the texture needed for our image sequence.
 	 *
 	 * @param InSequenceName Name of this sequence.
-	 * @param NumMipMaps Number of mipmaps in our image sequence.
-	 * @param Dim Dimensions of the textures in our image sequence.
+	 * @param InNumMipMaps Number of mipmaps in our image sequence.
+	 * @param InNumTiles Number of tiles in our image sequence.
+	 * @param InSequenceDim Dimensions of the textures in our image sequence.
 	 */
-	void SetTextureInfo(FName InSequenceName, int32 NumMipMaps, int32 InNumTilesX, int32 InNumTilesY,
-		const FIntPoint& Dim);
+	void SetTextureInfo(FName InSequenceName, int32 InNumMipMaps, const FIntPoint& InNumTiles,
+		const FIntPoint& InSequenceDim);
 
 	/**
 	 * Get what mipmap level should be used.
@@ -217,34 +252,7 @@ public:
 	 * @param TileSelection Will be set with which tiles are visible.
 	 * @return Mipmap level.
 	 */
-	int32 GetDesiredMipLevel(FImgMediaTileSelection& InTileSelection);
-
-	/**
-	 * Calculate object distance to camera. 
-	 * 
-	 * @param InCameraLocation Postion of camera.
-	 * @param InObjectLocation Position of object.
-	 * @return Distance.
-	 */
-	IMGMEDIA_API static float GetObjectDistToCamera(const FVector& InCameraLocation, const FVector& InObjectLocation);
-	
-	/**
-	 * Determine which mip level to use for a given distance.
-	 *
-	 * @param InDistance Distance to use.
-	 * @param InMipLevelDistances Distances for each mip level.
-	 * @return Mip level.
-	 */
-	IMGMEDIA_API static int GetMipLevelForDistance(float InDistance, const TArray<float>& InMipLevelDistances);
-
-	/**
-	 * Determine the size of an object.
-	 *
-	 * @param InActor Object to get the size of.
-	 * @param Width		Will be set to the width of the object.
-	 * @param Height	Will be set to the height of the object.
-	 */
-	void GetObjectSize(const AActor* InActor, float& Width, float& Height);
+	TMap<int32, FImgMediaTileSelection> GetVisibleTiles();
 	
 	/**
 	 * Get information on all our cameras.
@@ -253,14 +261,6 @@ public:
 	 */
 	const TArray<FImgMediaMipMapCameraInfo>& GetCameraInfo() const { return CameraInfos; }
 
-	/**
-	 * Get adjustment needed for distance to take the viewport size into account compared
-	 * to the reference viewport
-	 *
-	 * @return Adjustment factor.
-	 */
-	float GetViewportDistAdjust() const { return ViewportDistAdjust; }
-
 	/**
 	 * Get information on objects that are using our textures.
 	 * 
@@ -268,6 +268,13 @@ public:
 	 */
 	const TArray<FImgMediaMipMapObjectInfo*>& GetObjects() const { return Objects; }
 
+	/**
+	 * Check if any scene objects are using our img sequence.
+	 *
+	 * @return True if any object is active.
+	 */
+	bool HasObjects() const { return Objects.Num() > 0; }
+
 	//~ FTickableGameObject interface
 	virtual void Tick(float DeltaTime) override;
 	virtual TStatId GetStatId() const override { RETURN_QUICK_DECLARE_CYCLE_STAT(FImgMediaMipMapInfo, STATGROUP_Tickables); }
@@ -279,54 +286,21 @@ protected:
 	 */
 	void UpdateMipLevelCache();
 
-	/** Decide which tiles are visible and update CachedTileSelection.
-	 *
-	 * @param CameraInfo				View info to compute visibility from.
-	 * @param ObjectLocation			Position of object.
-	 * @param NextTileXVector			Vector to go from one tile to the next in the X axis.
-	 * @param NextTileYVector			Vector to go from one tile to the next in the Y axis.
-	 * @param TileRadiusInWorldSpace	Radius of a bounding sphere for a tile.
-	 */
-	void CalculateTileVisibility(const FImgMediaMipMapCameraInfo& CameraInfo,
-		const FVector& ObjectLocation, const FVector& NextTileXVector,
-		const FVector& NextTileYVector, float TileRadiusInWorldSpace);
-
-
-	/**
-	 * Updates the MipLevelDistances based on current information.
-	 */
-	void UpdateMipLevelDistances();
-
-	/** Name of this sequence. */
-	FName SequenceName;
-	/** Number of tiles in the X direction. */
-	int32 NumTilesX;
-	/** Number of tiles in the Y direction. */
-	int32 NumTilesY;
-
-	/** Ideal distance for mip level 0. */
-	float MipLevel0Distance;
-	/** True if MipLevel0Distance has been set manually. */
-	bool bIsMipLevel0DistanceSetManually;
-
-	/** Ideal distances for all mip maps. */
-	TArray<float> MipLevelDistances;
-
 	/** Array of objects that are using our img sequence. */
 	TArray<FImgMediaMipMapObjectInfo*> Objects;
 
-	/** Adjustment for current size of viewport, used in mipmap calculations. */
-	float ViewportDistAdjust;
 	/** Info for each camera, used in mipmap calculations. */
 	TArray<FImgMediaMipMapCameraInfo> CameraInfos;
+
+	/** Size, tiling and mips sequence information. */
+	FSequenceInfo SequenceInfo;
 	
 	/** Protects info variables. */
 	FCriticalSection InfoCriticalSection;
 
 	/** Desired mipmap level at this current time. */
-	int32 CachedMipLevel;
-	/** Desired tiles at this current tmie. */
-	FImgMediaTileSelection CachedTileSelection;
+	TMap<int32, FImgMediaTileSelection> CachedVisibleTiles;
+
 	/** True if the cached mipmap data has been calculated this frame. */
-	bool bIsCachedMipLevelValid;
+	bool bIsCacheValid;
 };
diff --git a/Engine/Plugins/Media/ImgMedia/Source/ImgMedia/Public/ImgMediaSource.h b/Engine/Plugins/Media/ImgMedia/Source/ImgMedia/Public/ImgMediaSource.h
--- a/Engine/Plugins/Media/ImgMedia/Source/ImgMedia/Public/ImgMediaSource.h
+++ b/Engine/Plugins/Media/ImgMedia/Source/ImgMedia/Public/ImgMediaSource.h
@@ -143,14 +143,6 @@ public:
 	UFUNCTION(BlueprintCallable, Category = "ImgMedia|ImgMediaSource")
 	void RemoveTargetObject(AActor* InActor);
 
-	/**
-	 * Manually set when mip level 0 should appear.
-	 *
-	 * @param Distance Furthest distance from the camera when mip level 0 should be at 100%.
-	 */
-	UFUNCTION(BlueprintCallable, Category = "ImgMedia|ImgMediaSource")
-	void SetMipLevelDistance(float Distance);
-
 	/**
 	 * Get our mipmap info object.
 	 *
diff --git a/Engine/Plugins/Media/ImgMedia/Source/ImgMediaEditor/Private/Customizations/ImgMediaSourceCustomization.h b/Engine/Plugins/Media/ImgMedia/Source/ImgMediaEditor/Private/Customizations/ImgMediaSourceCustomization.h
--- a/Engine/Plugins/Media/ImgMedia/Source/ImgMediaEditor/Private/Customizations/ImgMediaSourceCustomization.h
+++ b/Engine/Plugins/Media/ImgMedia/Source/ImgMediaEditor/Private/Customizations/ImgMediaSourceCustomization.h
@@ -16,7 +16,6 @@ class SEditableTextBox;
 
 class FImgMediaMipMapInfo;
 struct FImgMediaMipMapCameraInfo;
-struct FImgMediaMipMapObjectInfo;
 
 /**
  * Implements a details view customization for the UImgMediaSource class.
