ImgMedia: Fixed process images filenames being sort unfriendly.
#rnx

diff --git a/Engine/Plugins/Media/ImgMedia/Source/ImgMediaEditor/Private/Widgets/SImgMediaProcessImages.cpp b/Engine/Plugins/Media/ImgMedia/Source/ImgMediaEditor/Private/Widgets/SImgMediaProcessImages.cpp
--- a/Engine/Plugins/Media/ImgMedia/Source/ImgMediaEditor/Private/Widgets/SImgMediaProcessImages.cpp
+++ b/Engine/Plugins/Media/ImgMedia/Source/ImgMediaEditor/Private/Widgets/SImgMediaProcessImages.cpp
@@ -817,8 +817,7 @@ void SImgMediaProcessImages::HandleProcessing()
 					bool bEnableMips = Options->bEnableMipMapping;
 					bool bHasAlphaChannel = false;
 					FString OutPath = Options->OutputPath.Path;
-					FString FileName = FString::Format(TEXT("image{0}.exr"),
-						{CurrentFrameIndex});
+					FString FileName = FString::Printf(TEXT("image%05d.exr"), CurrentFrameIndex);
 					FString Name = FPaths::Combine(OutPath, FileName);
 
 					ProcessImageCustomRawData(RawData, Width, Height, BitDepth,
