ImgMedia: Fixed missing FViewport forward declaration.

#rnx

diff --git a/Engine/Plugins/Media/ImgMedia/Source/ImgMedia/Private/Assets/ImgMediaMipMapInfoManager.h b/Engine/Plugins/Media/ImgMedia/Source/ImgMedia/Private/Assets/ImgMediaMipMapInfoManager.h
--- a/Engine/Plugins/Media/ImgMedia/Source/ImgMedia/Private/Assets/ImgMediaMipMapInfoManager.h
+++ b/Engine/Plugins/Media/ImgMedia/Source/ImgMedia/Private/Assets/ImgMediaMipMapInfoManager.h
@@ -12,6 +12,7 @@
 
 class AActor;
 class FSceneView;
+class FViewport;
 class IMediaPlayer;
 class UCameraComponent;
 class UImgMediaPlaybackComponent;
