ImgMedia: Process images now processes in separate threads when using a player.
#rnx

diff --git a/Engine/Plugins/Media/ImgMedia/Source/ImgMediaEditor/Private/Widgets/SImgMediaProcessImages.cpp b/Engine/Plugins/Media/ImgMedia/Source/ImgMediaEditor/Private/Widgets/SImgMediaProcessImages.cpp
--- a/Engine/Plugins/Media/ImgMedia/Source/ImgMediaEditor/Private/Widgets/SImgMediaProcessImages.cpp
+++ b/Engine/Plugins/Media/ImgMedia/Source/ImgMediaEditor/Private/Widgets/SImgMediaProcessImages.cpp
@@ -761,6 +761,7 @@ void SImgMediaProcessImages::HandleProcessing()
 	// Are we processing?
 	if (bIsProcessing)
 	{
+		TRACE_CPUPROFILER_EVENT_SCOPE(SImgMediaProcessImages::HandleProcessing);
 		// We did not cancel yet?
 		bool bShouldExit = false;
 		if ((MediaPlayer != nullptr) && (bIsCancelling == false))
@@ -820,9 +821,14 @@ void SImgMediaProcessImages::HandleProcessing()
 					FString FileName = FString::Printf(TEXT("image%05d.exr"), CurrentFrameIndex);
 					FString Name = FPaths::Combine(OutPath, FileName);
 
-					ProcessImageCustomRawData(RawData, Width, Height, BitDepth,
+					Async(EAsyncExecution::Thread, [this, RawData = MoveTemp(RawData), Width, Height, BitDepth,
 						InTileWidth, InTileHeight, TileBorder, bEnableMips,
-						bHasAlphaChannel, Name);
+						bHasAlphaChannel, Name]() mutable
+					{
+						ProcessImageCustomRawData(RawData, Width, Height, BitDepth,
+							InTileWidth, InTileHeight, TileBorder, bEnableMips,
+							bHasAlphaChannel, Name);
+					});
 				}
 				else
 				{
