ImgMedia: Removed now unused single-file-per-tile logic in loader and readers, removed tile selection SetAllVisible() calls so clamping to the number of tiles isn't required anymore, provided missing defaults to FImgMediaFrameInfo in the generic reader.

#rnx

diff --git a/Engine/Plugins/Media/ImgMedia/Source/ImgMedia/Private/Assets/ImgMediaMipMapInfo.cpp b/Engine/Plugins/Media/ImgMedia/Source/ImgMedia/Private/Assets/ImgMediaMipMapInfo.cpp
--- a/Engine/Plugins/Media/ImgMedia/Source/ImgMedia/Private/Assets/ImgMediaMipMapInfo.cpp
+++ b/Engine/Plugins/Media/ImgMedia/Source/ImgMedia/Private/Assets/ImgMediaMipMapInfo.cpp
@@ -37,8 +37,10 @@ bool FImgMediaMipMapObjectInfo::CalculateVisibleTiles(const TArray<FImgMediaMipM
 	{
 		FImgMediaTileSelection& Selection = VisibleTiles.FindOrAdd(MipLevel);
 
-		int32 NumTilesX = FMath::Max(1, InSequenceInfo.NumTiles.X >> MipLevel);
-		int32 NumTilesY = FMath::Max(1, InSequenceInfo.NumTiles.Y >> MipLevel);
+		const int MipLevelDiv = 1 << MipLevel;
+
+		int32 NumTilesX = FMath::Max(1, FMath::CeilToInt(float(InSequenceInfo.NumTiles.X) / MipLevelDiv));
+		int32 NumTilesY = FMath::Max(1, FMath::CeilToInt(float(InSequenceInfo.NumTiles.Y) / MipLevelDiv));
 		Selection.SetVisibleRegion(0u, 0u, IntCastChecked<uint16>(NumTilesX), IntCastChecked<uint16>(NumTilesY));
 	}
 	return true;
@@ -135,9 +137,11 @@ namespace {
 				GetViewFrustumBounds(ViewFrustum, CameraInfo.ViewProjectionMatrix, false, false);
 
 				int32 MaxLevel = InSequenceInfo.NumMipLevels - 1;
+				int MipLevelDiv = 1 << MaxLevel;
+
 				FIntPoint CurrentNumTiles;
-				CurrentNumTiles.X = FMath::Max(1, InSequenceInfo.NumTiles.X >> MaxLevel);
-				CurrentNumTiles.Y = FMath::Max(1, InSequenceInfo.NumTiles.Y >> MaxLevel);
+				CurrentNumTiles.X = FMath::Max(1, FMath::CeilToInt(float(InSequenceInfo.NumTiles.X) / MipLevelDiv));
+				CurrentNumTiles.Y = FMath::Max(1, FMath::CeilToInt(float(InSequenceInfo.NumTiles.Y) / MipLevelDiv));
 
 				// Starting with tiles at the highest mip level
 				TQueue<FIntVector> Tiles;
@@ -156,9 +160,10 @@ namespace {
 					Tiles.Dequeue(Tile);
 
 					int32 CurrentMipLevel = Tile.Z;
+					MipLevelDiv = 1 << CurrentMipLevel;
 					// Calculate the number of tiles at this mip level
-					CurrentNumTiles.X = FMath::Max(1, InSequenceInfo.NumTiles.X >> CurrentMipLevel);
-					CurrentNumTiles.Y = FMath::Max(1, InSequenceInfo.NumTiles.Y >> CurrentMipLevel);
+					CurrentNumTiles.X = FMath::Max(1, FMath::CeilToInt(float(InSequenceInfo.NumTiles.X) / MipLevelDiv));
+					CurrentNumTiles.Y = FMath::Max(1, FMath::CeilToInt(float(InSequenceInfo.NumTiles.Y) / MipLevelDiv));
 
 					// Calculate the tile location in world-space
 					float StepX = float(Tile.X + 0.5f) / CurrentNumTiles.X;
diff --git a/Engine/Plugins/Media/ImgMedia/Source/ImgMedia/Private/Loader/ImgMediaLoader.cpp b/Engine/Plugins/Media/ImgMedia/Source/ImgMedia/Private/Loader/ImgMediaLoader.cpp
--- a/Engine/Plugins/Media/ImgMedia/Source/ImgMedia/Private/Loader/ImgMediaLoader.cpp
+++ b/Engine/Plugins/Media/ImgMedia/Source/ImgMedia/Private/Loader/ImgMediaLoader.cpp
@@ -754,8 +754,8 @@ void FImgMediaLoader::LoadSequence(const FString& SequencePath, const FFrameRate
 	// Get mips.
 	FindMips(SequencePath);
 
-	// Get tile info.
-	GetTileInformation();
+	// Default to a 1x1 tile
+	TilingDescription.TileNum = FIntPoint(1, 1);
 
 	FScopeLock Lock(&CriticalSection);
 
@@ -969,12 +969,6 @@ void FImgMediaLoader::FindFiles(const FString& SequencePath, TArray<FString>& Ou
 		}
 
 		FString FullPath = FPaths::Combine(SequencePath, File);
-		if (bIsTiled)
-		{
-			// If we have tiles, then use the first tile as our reference.
-			FString TiledFileName = File + TEXT("_x0_y0.exr");
-			FullPath = FPaths::Combine(FullPath, TiledFileName);
-		}
 		OutputPaths.Add(FullPath);
 	}
 
@@ -1056,49 +1050,6 @@ void FImgMediaLoader::FindMips(const FString& SequencePath)
 	}
 }
 
-void FImgMediaLoader::GetTileInformation()
-{
-	TilingDescription.TileNum.X = 1;
-	TilingDescription.TileNum.Y = 1;
-
-	// Do we have tiles?
-	if (bIsTiled)
-	{
-		FString Path = FPaths::GetPath(ImagePaths[0][0]);
-		TArray<FString> FoundFiles;
-		IFileManager::Get().FindFiles(FoundFiles, *Path, TEXT("*"));
-
-		for (const FString& FilePath : FoundFiles)
-		{
-			// Remove the file extension.
-			Path = FPaths::ChangeExtension(FilePath, TEXT(""));
-
-			// Filename should be ????_x0_y0.
-			// Find the last y.
-			int32 YIndex = 0;
-			if (Path.FindLastChar(TEXT('y'), YIndex))
-			{
-				// Extract out the number.
-				int32 YLength = Path.Len() - (YIndex + 1);
-				FString NumberString = Path.Right(YLength);
-				// We count from 0, so add 1 to get the number of tiles.
-				TilingDescription.TileNum.Y = FMath::Max(TilingDescription.TileNum.Y, FCString::Atoi(*NumberString) + 1);
-
-				// Find the last x.
-				int32 XIndex = 0;
-				if (Path.FindLastChar(TEXT('x'), XIndex))
-				{
-					// Extract out the number.
-					XIndex++;
-					int32 XLength = Path.Len() - XIndex - YLength - 2;
-					NumberString = Path.Mid(XIndex, XLength);
-					TilingDescription.TileNum.X = FMath::Max(TilingDescription.TileNum.X, FCString::Atoi(*NumberString) + 1);
-				}
-			}
-		}
-	}
-}
-
 uint32 FImgMediaLoader::TimeToFrameNumber(FTimespan Time) const
 {
 	if ((Time < FTimespan::Zero()) || (Time >= SequenceDuration))
@@ -1302,7 +1253,7 @@ void FImgMediaLoader::AddEmptyFrame(int32 FrameNumber)
 	Frame->Stride = Frame ->Info.Dim.X * PixelSize;
 	for (int32 Level = 0; Level < GetNumMipLevels(); ++Level)
 	{
-		Frame->MipTilesPresent.Add(Level).SetAllVisible();
+		Frame->MipTilesPresent.Add(Level).Include(0u, 0u);
 	}
 	AddFrameToCache(FrameNumber, Frame);
 }
@@ -1378,8 +1329,10 @@ void FImgMediaLoader::GetDesiredMipTiles(int32 FrameIndex, TMap<int32, FImgMedia
 		{
 			FImgMediaTileSelection& Selection = OutMipsAndTiles.FindOrAdd(MipLevel);
 
-			int32 NumTilesX = FMath::Max(1, GetNumTilesX() >> MipLevel);
-			int32 NumTilesY = FMath::Max(1, GetNumTilesY() >> MipLevel);
+			const int MipLevelDiv = 1 << MipLevel;
+
+			int32 NumTilesX = FMath::Max(1, FMath::CeilToInt(float(GetNumTilesX()) / MipLevelDiv));
+			int32 NumTilesY = FMath::Max(1, FMath::CeilToInt(float(GetNumTilesY()) / MipLevelDiv));
 			Selection.SetVisibleRegion(0u, 0u, IntCastChecked<uint16>(NumTilesX), IntCastChecked<uint16>(NumTilesY));
 		}
 	}
diff --git a/Engine/Plugins/Media/ImgMedia/Source/ImgMedia/Private/Loader/ImgMediaLoader.h b/Engine/Plugins/Media/ImgMedia/Source/ImgMedia/Private/Loader/ImgMediaLoader.h
--- a/Engine/Plugins/Media/ImgMedia/Source/ImgMedia/Private/Loader/ImgMediaLoader.h
+++ b/Engine/Plugins/Media/ImgMedia/Source/ImgMedia/Private/Loader/ImgMediaLoader.h
@@ -372,11 +372,6 @@ protected:
 	 */
 	void FindMips(const FString& SequencePath);
 
-	/**
-	 * Gets information for tiling from the sequence.
-	 */
-	void GetTileInformation();
-
 	/**
 	 * Get the frame number corresponding to the specified play head time.
 	 *
diff --git a/Engine/Plugins/Media/ImgMedia/Source/ImgMedia/Private/Readers/ExrImgMediaReader.cpp b/Engine/Plugins/Media/ImgMedia/Source/ImgMedia/Private/Readers/ExrImgMediaReader.cpp
--- a/Engine/Plugins/Media/ImgMedia/Source/ImgMedia/Private/Readers/ExrImgMediaReader.cpp
+++ b/Engine/Plugins/Media/ImgMedia/Source/ImgMedia/Private/Readers/ExrImgMediaReader.cpp
@@ -149,11 +149,6 @@ bool FExrImgMediaReader::ReadFrame(int32 FrameId, const TMap<int32, FImgMediaTil
 			return false;
 		}
 
-		// If we have tiles, then this is the size of just a tile, so multiply to get the full size. Same goes for size.
-		OutFrame->Info.Dim.X *= NumTilesX;
-		OutFrame->Info.Dim.Y *= NumTilesY;
-		OutFrame->Info.UncompressedSize *= NumTilesX * NumTilesY;
-
 		const FIntPoint& Dim = OutFrame->Info.Dim;
 
 		if (Dim.GetMin() <= 0)
@@ -205,10 +200,6 @@ bool FExrImgMediaReader::ReadFrame(int32 FrameId, const TMap<int32, FImgMediaTil
 			const FImgMediaTileSelection& CurrentTileSelection = InMipTiles[CurrentMipLevel];
 
 			const int MipLevelDiv = 1 << CurrentMipLevel;
-			int32 StartTileX = CurrentTileSelection.TopLeftX;
-			int32 StartTileY = CurrentTileSelection.TopLeftY;
-			int32 EndTileX = FMath::Min((int32)CurrentTileSelection.BottomRightX, FMath::CeilToInt(float(NumTilesX) / MipLevelDiv));
-			int32 EndTileY = FMath::Min((int32)CurrentTileSelection.BottomRightY, FMath::CeilToInt(float(NumTilesY) / MipLevelDiv));
 
 			bool ReadThisMip = true;
 			// Avoid reads if the cached frame already contains the current tiles for this mip level.
@@ -221,18 +212,6 @@ bool FExrImgMediaReader::ReadFrame(int32 FrameId, const TMap<int32, FImgMediaTil
 			{
 				FString Image = Loader->GetImagePath(FrameId, CurrentMipLevel);
 				FString BaseImage;
-				if (bHasTiles)
-				{
-					// Remove "_x0_y0.exr" so we can add on the correct name for the tile we want.
-					BaseImage = Image.LeftChop(10);
-				}
-
-				int32 TileWidth = Dim.X / NumTilesX;
-				int32 TileHeight = Dim.Y / NumTilesY;
-				// The offset into the frame buffer for each row of tiles.
-				// Total width * height of a tile * bytes per pixel.
-				int32 FrameBufferOffsetPerTileY = Dim.X * TileHeight * BytesPerPixel;
-				int32 FrameBufferOffsetY = FrameBufferOffsetPerTileY * StartTileY;
 
 				if (OutFrame->Info.FormatName == TEXT("EXR CUSTOM"))
 				{
@@ -240,8 +219,8 @@ bool FExrImgMediaReader::ReadFrame(int32 FrameId, const TMap<int32, FImgMediaTil
 					FIntRect TileRegion = FIntRect(
 						(int32)CurrentTileSelection.TopLeftX,
 						(int32)CurrentTileSelection.TopLeftY,
-						FMath::Min((int32)CurrentTileSelection.BottomRightX, FMath::CeilToInt(float(OutFrame->Info.NumTiles.X) / MipLevelDiv)),
-						FMath::Min((int32)CurrentTileSelection.BottomRightY, FMath::CeilToInt(float(OutFrame->Info.NumTiles.Y) / MipLevelDiv)));
+						(int32)CurrentTileSelection.BottomRightX,
+						(int32)CurrentTileSelection.BottomRightY);
 					FIntPoint NumTiles = OutFrame->Info.NumTiles / MipLevelDiv;
 					int32 PixelSize = sizeof(uint16) * OutFrame->Info.NumChannels;
 					TSharedPtr<FSampleConverterParameters> ConverterParams = MakeShared<FSampleConverterParameters>();
@@ -264,42 +243,24 @@ bool FExrImgMediaReader::ReadFrame(int32 FrameId, const TMap<int32, FImgMediaTil
 					UE_LOG(LogImgMedia, Error, TEXT("Current platform doesn't support custom EXR file %s"), *Image);
 #endif
 				}
+				else if (bHasTiles)
+				{
+					UE_LOG(LogImgMedia, Error, TEXT("Non-GPU reader doesn't currently support natively-tiled EXR file %s"), *Image);
+				}
 				else
 				{
-					for (int32 TileY = StartTileY; TileY < EndTileY; TileY++)
+					// Get for our frame/mip level.
+					FRgbaInputFile InputFile(Image, 2);
+					if (InputFile.HasInputFile())
 					{
-						// The offset into the frame buffer for each column of tiles.
-						// Tile width * bytes per pixel.
-						int32 FrameBufferOffsetPerTileX = TileWidth * BytesPerPixel;
-						int32 FrameBufferOffsetX = FrameBufferOffsetPerTileX * StartTileX;
-						for (int32 TileX = StartTileX; TileX < EndTileX; TileX++)
-						{
-							// Get for our frame/mip level.
-							if (bHasTiles)
-							{
-								Image = FString::Printf(TEXT("%s_x%d_y%d.exr"), *BaseImage, TileX, TileY);
-							}
-							FRgbaInputFile InputFile(Image, 2);
-							if (InputFile.HasInputFile())
-							{
-								// read frame data
-								InputFile.SetFrameBuffer(MipDataPtr + FrameBufferOffsetX + FrameBufferOffsetY, Dim);
-								InputFile.ReadPixels(0, TileHeight - 1);
-
-								FrameBufferOffsetX += FrameBufferOffsetPerTileX;
-
-								if (!OutFrame->MipTilesPresent.Contains(CurrentMipLevel))
-								{
-									OutFrame->MipTilesPresent.Emplace(CurrentMipLevel, CurrentTileSelection);
-								}
-							}
-							else
-							{
-								UE_LOG(LogImgMedia, Error, TEXT("Could not load %s"), *Image);
-							}
-						}
-
-						FrameBufferOffsetY += FrameBufferOffsetPerTileY;
+						// read frame data
+						InputFile.SetFrameBuffer(MipDataPtr, Dim);
+						InputFile.ReadPixels(0, Dim.Y - 1);
+						OutFrame->MipTilesPresent.Emplace(CurrentMipLevel, CurrentTileSelection);
+					}
+					else
+					{
+						UE_LOG(LogImgMedia, Error, TEXT("Could not load %s"), *Image);
 					}
 				}
 			}
diff --git a/Engine/Plugins/Media/ImgMedia/Source/ImgMedia/Private/Readers/ExrImgMediaReaderGpu.cpp b/Engine/Plugins/Media/ImgMedia/Source/ImgMedia/Private/Readers/ExrImgMediaReaderGpu.cpp
--- a/Engine/Plugins/Media/ImgMedia/Source/ImgMedia/Private/Readers/ExrImgMediaReaderGpu.cpp
+++ b/Engine/Plugins/Media/ImgMedia/Source/ImgMedia/Private/Readers/ExrImgMediaReaderGpu.cpp
@@ -204,8 +204,8 @@ bool FExrImgMediaReaderGpu::ReadFrame(int32 FrameId, const TMap<int32, FImgMedia
 						FIntRect TileRegion = FIntRect(
 							(int32)CurrentTileSelection.TopLeftX,
 							(int32)CurrentTileSelection.TopLeftY,
-							FMath::Min((int32)CurrentTileSelection.BottomRightX, FMath::CeilToInt(float(OutFrame->Info.NumTiles.X) / MipLevelDiv)),
-							FMath::Min((int32)CurrentTileSelection.BottomRightY, FMath::CeilToInt(float(OutFrame->Info.NumTiles.Y) / MipLevelDiv)));
+							(int32)CurrentTileSelection.BottomRightX,
+							(int32)CurrentTileSelection.BottomRightY);
 
 						ReadResult = ReadTilesCustom(MipDataPtr, ImagePath, FrameId, TileRegion, ConverterParams, CurrentMipLevel);
 					}
diff --git a/Engine/Plugins/Media/ImgMedia/Source/ImgMedia/Private/Readers/GenericImgMediaReader.cpp b/Engine/Plugins/Media/ImgMedia/Source/ImgMedia/Private/Readers/GenericImgMediaReader.cpp
--- a/Engine/Plugins/Media/ImgMedia/Source/ImgMedia/Private/Readers/GenericImgMediaReader.cpp
+++ b/Engine/Plugins/Media/ImgMedia/Source/ImgMedia/Private/Readers/GenericImgMediaReader.cpp
@@ -65,6 +65,12 @@ TSharedPtr<IImageWrapper> LoadImage(const FString& ImagePath, IImageWrapperModul
 	OutInfo.FrameRate = Settings->DefaultFrameRate;
 	OutInfo.Srgb = true;
 	OutInfo.UncompressedSize = (SIZE_T)OutInfo.Dim.X * OutInfo.Dim.Y * 4;
+	OutInfo.NumChannels = 4;
+	OutInfo.bHasTiles = false;
+	OutInfo.TileDimensions = OutInfo.Dim;
+	OutInfo.NumTiles = FIntPoint(1, 1);
+	OutInfo.TileBorder = 0;
+	OutInfo.NumMipLevels = 1;
 
 	return ImageWrapper;
 }
@@ -109,6 +115,7 @@ bool FGenericImgMediaReader::ReadFrame(int32 FrameId, const TMap<int32, FImgMedi
 	for (const TPair<int32, FImgMediaTileSelection>& TilesPerMip : InMipTiles)
 	{
 		const int32 CurrentMipLevel = TilesPerMip.Key;
+		const FImgMediaTileSelection& CurrentSelection = TilesPerMip.Value;
 
 		// Do we want to read in this mip?
 		bool IsThisLevelPresent = OutFrame->MipTilesPresent.Contains(CurrentMipLevel);
@@ -154,7 +161,7 @@ bool FGenericImgMediaReader::ReadFrame(int32 FrameId, const TMap<int32, FImgMedi
 
 			// Copy data to our buffer
 			FMemory::Memcpy(Buffer, RawData.GetData(), RawNum);
-			OutFrame->MipTilesPresent.Add(CurrentMipLevel).SetAllVisible();
+			OutFrame->MipTilesPresent.Emplace(CurrentMipLevel, CurrentSelection);
 		}
 
 		// Next level.
