ImgMedia: Tile selection object now registers individual visible tiles in a bit array covering the entire region.

diff --git a/Engine/Plugins/Media/ImgMedia/Source/ImgMedia/Private/Assets/ImgMediaMipMapInfo.cpp b/Engine/Plugins/Media/ImgMedia/Source/ImgMedia/Private/Assets/ImgMediaMipMapInfo.cpp
--- a/Engine/Plugins/Media/ImgMedia/Source/ImgMedia/Private/Assets/ImgMediaMipMapInfo.cpp
+++ b/Engine/Plugins/Media/ImgMedia/Source/ImgMedia/Private/Assets/ImgMediaMipMapInfo.cpp
@@ -19,6 +19,147 @@
 
 DECLARE_CYCLE_STAT(TEXT("ImgMedia MipMap Update Cache"), STAT_ImgMedia_MipMapUpdateCache, STATGROUP_Media);
 
+FImgMediaTileSelection::FImgMediaTileSelection(int32 NumTilesX, int32 NumTilesY, bool bDefaultVisibility)
+	: Tiles(bDefaultVisibility, NumTilesX * NumTilesY)
+	, Dimensions(NumTilesX, NumTilesY)
+	, CachedVisibleRegion()
+	, bCachedVisibleRegionDirty(true)
+{
+}
+
+FImgMediaTileSelection FImgMediaTileSelection::CreateForTargetMipLevel(int32 BaseNumTilesX, int32 BaseNumTilesY, int32 TargetMipLevel, bool bDefaultVisibility)
+{
+	ensure(TargetMipLevel >= 0);
+
+	const int MipLevelDiv = 1 << TargetMipLevel;
+	int32 NumTilesX = FMath::Max(1, FMath::CeilToInt(float(BaseNumTilesX) / MipLevelDiv));
+	int32 NumTilesY = FMath::Max(1, FMath::CeilToInt(float(BaseNumTilesY) / MipLevelDiv));
+
+	return FImgMediaTileSelection(NumTilesX, NumTilesY, bDefaultVisibility);
+}
+
+bool FImgMediaTileSelection::IsAnyVisible() const
+{
+	return Tiles.Contains(true);
+};
+
+bool FImgMediaTileSelection::IsVisible(int32 TileCoordX, int32 TileCoordY) const
+{
+	return Tiles[ToIndex(TileCoordX, TileCoordY, Dimensions)];
+}
+
+bool FImgMediaTileSelection::Contains(const FImgMediaTileSelection& Other) const
+{
+	//Modified version of TBitArray's CompareSetBits() method.
+
+	TBitArray<>::FConstWordIterator ThisIterator(Tiles);
+	TBitArray<>::FConstWordIterator OtherIterator(Other.Tiles);
+
+	ThisIterator.FillMissingBits(0u);
+	OtherIterator.FillMissingBits(0u);
+
+	while (ThisIterator || OtherIterator)
+	{
+		const uint32 A = ThisIterator ? ThisIterator.GetWord() : 0u;
+		const uint32 B = OtherIterator ? OtherIterator.GetWord() : 0u;
+		if (A != B)
+		{
+			// Check if A contains all of the ones in B.
+			if ((A & B) != B)
+			{
+				return false;
+			}
+		}
+
+		++ThisIterator;
+		++OtherIterator;
+	}
+
+	return true;
+}
+
+void FImgMediaTileSelection::SetVisible(int32 TileCoordX, int32 TileCoordY)
+{
+	Tiles[ToIndex(TileCoordX, TileCoordY, Dimensions)] = true;
+	bCachedVisibleRegionDirty = true;
+}
+
+TArray<FIntPoint> FImgMediaTileSelection::GetVisibleCoordinates() const
+{
+	TArray<FIntPoint> OutCoordinates;
+
+	for (int32 CoordY = 0; CoordY < Dimensions.Y; ++CoordY)
+	{
+		for (int32 CoordX = 0; CoordX < Dimensions.X; ++CoordX)
+		{
+			if (Tiles[ToIndex(CoordX, CoordY, Dimensions)])
+			{
+				OutCoordinates.Emplace(CoordX, CoordY);
+			}
+		}
+	}
+
+	return OutCoordinates;
+}
+
+FIntRect FImgMediaTileSelection::GetVisibleRegion() const
+{
+	// We offload the region calculation to the loader workers, instead of constantly updating it during SetVisible().
+	// Not thread safe, but only accessed sequentially in individual worker thread copies.
+
+	if (bCachedVisibleRegionDirty)
+	{
+		FIntPoint Min = TNumericLimits<int32>::Max();
+		FIntPoint Max = TNumericLimits<int32>::Min();
+
+		for (int32 CoordY = 0; CoordY < Dimensions.Y; ++CoordY)
+		{
+			for (int32 CoordX = 0; CoordX < Dimensions.X; ++CoordX)
+			{
+				if (Tiles[ToIndex(CoordX, CoordY, Dimensions)])
+				{
+					Min.X = FMath::Min(Min.X, CoordX);
+					Min.Y = FMath::Min(Min.Y, CoordY);
+					Max.X = FMath::Max(Max.X, CoordX);
+					Max.Y = FMath::Max(Max.Y, CoordY);
+				}
+			}
+		}
+
+		if (Max.X >= Min.X && Max.Y >= Min.Y)
+		{
+			CachedVisibleRegion = FIntRect(Min, Max + 1);
+		}
+		else
+		{
+			CachedVisibleRegion = FIntRect();
+		}
+
+		bCachedVisibleRegionDirty = false;
+	}
+
+	return CachedVisibleRegion;
+}
+
+int32 FImgMediaTileSelection::NumVisibleTiles() const
+{
+	int32 NumVisibleTiles = 0;
+
+	for (int32 CoordY = 0; CoordY < Dimensions.Y; ++CoordY)
+	{
+		for (int32 CoordX = 0; CoordX < Dimensions.X; ++CoordX)
+		{
+			if (Tiles[ToIndex(CoordX, CoordY, Dimensions)])
+			{
+				++NumVisibleTiles;
+			}
+		}
+	}
+
+	return NumVisibleTiles;
+}
+
+
 FImgMediaMipMapObjectInfo::FImgMediaMipMapObjectInfo(UMeshComponent* InMeshComponent, float InLODBias)
 	: MeshComponent(InMeshComponent)
 	, LODBias(InLODBias)
@@ -33,15 +174,10 @@ UMeshComponent* FImgMediaMipMapObjectInfo::GetMeshComponent() const
 
 bool FImgMediaMipMapObjectInfo::CalculateVisibleTiles(const TArray<FImgMediaMipMapCameraInfo>& InCameraInfos, const FSequenceInfo& InSequenceInfo, TMap<int32, FImgMediaTileSelection>& VisibleTiles) const
 {
+	// We simply add fully visible regions for all mip levels
 	for (int32 MipLevel = 0; MipLevel < InSequenceInfo.NumMipLevels; ++MipLevel)
 	{
-		FImgMediaTileSelection& Selection = VisibleTiles.FindOrAdd(MipLevel);
-
-		const int MipLevelDiv = 1 << MipLevel;
-
-		int32 NumTilesX = FMath::Max(1, FMath::CeilToInt(float(InSequenceInfo.NumTiles.X) / MipLevelDiv));
-		int32 NumTilesY = FMath::Max(1, FMath::CeilToInt(float(InSequenceInfo.NumTiles.Y) / MipLevelDiv));
-		Selection.SetVisibleRegion(0u, 0u, IntCastChecked<uint16>(NumTilesX), IntCastChecked<uint16>(NumTilesY));
+		VisibleTiles.Add(MipLevel, FImgMediaTileSelection::CreateForTargetMipLevel(InSequenceInfo.NumTiles.X, InSequenceInfo.NumTiles.Y, MipLevel, true));
 	}
 	return true;
 }
@@ -243,7 +379,12 @@ namespace {
 						// If the highest (calculated) mip level equals or exceeds our current mip level, we register the tile as visible.
 						if (MipLevelRange[1] >= CurrentMipLevel)
 						{
-							VisibleTiles.FindOrAdd(CurrentMipLevel).Include(Tile.X, Tile.Y);
+							if (!VisibleTiles.Contains(CurrentMipLevel))
+							{
+								VisibleTiles.Emplace(CurrentMipLevel, FImgMediaTileSelection(CurrentNumTiles.X, CurrentNumTiles.Y));
+							}
+
+							VisibleTiles[CurrentMipLevel].SetVisible(Tile.X, Tile.Y);
 						}
 #if false
 #if WITH_EDITOR
@@ -294,25 +435,6 @@ namespace {
 		mutable TArray<float> CornerMipLevelsCached;
 	};
 
-	// Reduce the selected tiles region for higher mip levels, always ensuring that all tiles at the base level are included.
-	FImgMediaTileSelection DownscaleTileSelection(const FImgMediaTileSelection& InTopLevelSelection, int32 InMipLevel)
-	{
-		if (InMipLevel <= 0)
-		{
-			return InTopLevelSelection;
-		}
-
-		const int32 MipLevelDiv = 1 << InMipLevel;
-
-		FImgMediaTileSelection OutSelection;
-		OutSelection.TopLeftX = (uint16)(InTopLevelSelection.TopLeftX / MipLevelDiv);
-		OutSelection.TopLeftY = (uint16)(InTopLevelSelection.TopLeftY / MipLevelDiv);
-		OutSelection.BottomRightX = (uint16)FMath::CeilToInt(float(InTopLevelSelection.BottomRightX) / MipLevelDiv);
-		OutSelection.BottomRightY = (uint16)FMath::CeilToInt(float(InTopLevelSelection.BottomRightY) / MipLevelDiv);
-
-		return OutSelection;
-	}
-
 	class FSphereObjectInfo : public FImgMediaMipMapObjectInfo
 	{
 	public:
@@ -335,13 +457,13 @@ namespace {
 				GetViewFrustumBounds(ViewFrustum, CameraInfo.ViewProjectionMatrix, false, false);
 				
 				// Include all tiles containted in the visible UV region
-				uint16 NumX = IntCastChecked<uint16>(InSequenceInfo.NumTiles.X);
-				uint16 NumY = IntCastChecked<uint16>(InSequenceInfo.NumTiles.Y);
-				for (uint16 TileY = 0; TileY < NumY; ++TileY)
+				int32 NumX = InSequenceInfo.NumTiles.X;
+				int32 NumY = InSequenceInfo.NumTiles.Y;
+				for (int32 TileY = 0; TileY < NumY; ++TileY)
 				{
-					for (uint16 TileX = 0; TileX < NumX; ++TileX)
+					for (int32 TileX = 0; TileX < NumX; ++TileX)
 					{
-						FVector2D TileCornerUV = FVector2D((float)TileX / InSequenceInfo.NumTiles.X, (float)TileY / InSequenceInfo.NumTiles.Y);
+						FVector2D TileCornerUV = FVector2D((float)TileX / NumX, (float)TileY / NumY);
 						
 						// Convert from latlong UV to spherical coordinates
 						FVector2D TileCornerSpherical = FVector2D(UE_PI * TileCornerUV.Y, UE_TWO_PI * TileCornerUV.X);
@@ -355,23 +477,42 @@ namespace {
 						// For each tile corner, we include all adjacent tiles
 						if (ViewFrustum.IntersectPoint(TileCorner))
 						{
-							VisibleTiles.FindOrAdd(0).Include(TileX, TileY);
+							if (!VisibleTiles.Contains(0))
+							{
+								VisibleTiles.Emplace(0, FImgMediaTileSelection(NumX, NumY));
+							}
 
-							int AdjacentX = TileX > 0 ? TileX - 1 : NumX - 1;
-							int AdjacentY = TileY > 0 ? TileY - 1 : NumY - 1;
-							VisibleTiles[0].Include(AdjacentX, TileY);
-							VisibleTiles[0].Include(TileX, AdjacentY);
-							VisibleTiles[0].Include(AdjacentX, AdjacentY);
+							int32 AdjacentX = TileX > 0 ? TileX - 1 : NumX - 1;
+							int32 AdjacentY = TileY > 0 ? TileY - 1 : NumY - 1;
+							VisibleTiles[0].SetVisible(TileX, TileY);
+							VisibleTiles[0].SetVisible(AdjacentX, TileY);
+							VisibleTiles[0].SetVisible(TileX, AdjacentY);
+							VisibleTiles[0].SetVisible(AdjacentX, AdjacentY);
 						}
 					}
 				}
 
 				if (VisibleTiles.Contains(0))
 				{
+					FIntPoint BaseDim = VisibleTiles[0].GetDimensions();
+					TArray<FIntPoint> BaseVisibleCoordinates = VisibleTiles[0].GetVisibleCoordinates();
+
+					// Include tiles visible at the base level in higher mip levels.
 					for (int32 Level = 1; Level < InSequenceInfo.NumMipLevels; ++Level)
 					{
-						FImgMediaTileSelection& Selection = VisibleTiles.FindOrAdd(Level);
-						Selection = DownscaleTileSelection(VisibleTiles[0], Level);
+						const int32 MipLevelDiv = 1 << Level;
+
+						if (!VisibleTiles.Contains(Level))
+						{
+							NumX = FMath::CeilToInt((float)BaseDim.X / MipLevelDiv);
+							NumY = FMath::CeilToInt((float)BaseDim.Y / MipLevelDiv);
+							VisibleTiles.Emplace(Level, FImgMediaTileSelection(NumX, NumY));
+						}
+
+						for (const FIntPoint& Coord : BaseVisibleCoordinates)
+						{
+							VisibleTiles[Level].SetVisible(Coord.X / MipLevelDiv, Coord.Y / MipLevelDiv);
+						}
 					}
 				}
 			}
@@ -546,7 +687,7 @@ void FImgMediaMipMapInfo::Tick(float DeltaTime)
 				MipLevelRange[1] = FMath::Max(MipLevelRange[1], MipTiles.Key);
 
 				const FImgMediaTileSelection& TileSelection = MipTiles.Value;
-				NumVisibleTiles += (TileSelection.BottomRightX - TileSelection.TopLeftX) * (TileSelection.BottomRightY - TileSelection.TopLeftY);
+				NumVisibleTiles += TileSelection.NumVisibleTiles();
 			}
 
 			if (MipLevelRange[0] <= MipLevelRange[1])
diff --git a/Engine/Plugins/Media/ImgMedia/Source/ImgMedia/Private/Loader/ImgMediaLoader.cpp b/Engine/Plugins/Media/ImgMedia/Source/ImgMedia/Private/Loader/ImgMediaLoader.cpp
--- a/Engine/Plugins/Media/ImgMedia/Source/ImgMedia/Private/Loader/ImgMediaLoader.cpp
+++ b/Engine/Plugins/Media/ImgMedia/Source/ImgMedia/Private/Loader/ImgMediaLoader.cpp
@@ -1253,7 +1253,7 @@ void FImgMediaLoader::AddEmptyFrame(int32 FrameNumber)
 	Frame->Stride = Frame ->Info.Dim.X * PixelSize;
 	for (int32 Level = 0; Level < GetNumMipLevels(); ++Level)
 	{
-		Frame->MipTilesPresent.Add(Level).Include(0u, 0u);
+		Frame->MipTilesPresent.Emplace(Level, FImgMediaTileSelection(1, 1, true));
 	}
 	AddFrameToCache(FrameNumber, Frame);
 }
@@ -1309,13 +1309,13 @@ void FImgMediaLoader::GetDesiredMipTiles(int32 FrameIndex, TMap<int32, FImgMedia
 
 		for (const auto& Pair : VisibleTiles)
 		{
-			FImgMediaTileSelection& OutSelection = OutMipsAndTiles.Add(Pair.Key);
-			OutSelection.SetAllNotVisible();
+			FImgMediaTileSelection Selection = FImgMediaTileSelection::CreateForTargetMipLevel(GetNumTilesX(), GetNumTilesY(), Pair.Key);
 
 			for (const FMediaTileCoordinate& Coordinate : Pair.Value)
 			{
-				OutSelection.Include(Coordinate.X, Coordinate.Y);
+				Selection.SetVisible(Coordinate.X, Coordinate.Y);
 			}
+			OutMipsAndTiles.Emplace(Pair.Key, MoveTemp(Selection));
 		}
 	}
 	else if(MipMapInfo.IsValid() && MipMapInfo->HasObjects())
@@ -1327,13 +1327,7 @@ void FImgMediaLoader::GetDesiredMipTiles(int32 FrameIndex, TMap<int32, FImgMedia
 		//Fallback case where we activate all mips and tiles.
 		for (int32 MipLevel = 0; MipLevel < GetNumMipLevels(); ++MipLevel)
 		{
-			FImgMediaTileSelection& Selection = OutMipsAndTiles.FindOrAdd(MipLevel);
-
-			const int MipLevelDiv = 1 << MipLevel;
-
-			int32 NumTilesX = FMath::Max(1, FMath::CeilToInt(float(GetNumTilesX()) / MipLevelDiv));
-			int32 NumTilesY = FMath::Max(1, FMath::CeilToInt(float(GetNumTilesY()) / MipLevelDiv));
-			Selection.SetVisibleRegion(0u, 0u, IntCastChecked<uint16>(NumTilesX), IntCastChecked<uint16>(NumTilesY));
+			OutMipsAndTiles.Emplace(MipLevel, FImgMediaTileSelection::CreateForTargetMipLevel(GetNumTilesX(), GetNumTilesY(), MipLevel, true));
 		}
 	}
 }
diff --git a/Engine/Plugins/Media/ImgMedia/Source/ImgMedia/Private/Readers/ExrImgMediaReader.cpp b/Engine/Plugins/Media/ImgMedia/Source/ImgMedia/Private/Readers/ExrImgMediaReader.cpp
--- a/Engine/Plugins/Media/ImgMedia/Source/ImgMedia/Private/Readers/ExrImgMediaReader.cpp
+++ b/Engine/Plugins/Media/ImgMedia/Source/ImgMedia/Private/Readers/ExrImgMediaReader.cpp
@@ -216,11 +216,7 @@ bool FExrImgMediaReader::ReadFrame(int32 FrameId, const TMap<int32, FImgMediaTil
 				if (OutFrame->Info.FormatName == TEXT("EXR CUSTOM"))
 				{
 #if defined(PLATFORM_WINDOWS) && PLATFORM_WINDOWS
-					FIntRect TileRegion = FIntRect(
-						(int32)CurrentTileSelection.TopLeftX,
-						(int32)CurrentTileSelection.TopLeftY,
-						(int32)CurrentTileSelection.BottomRightX,
-						(int32)CurrentTileSelection.BottomRightY);
+					FIntRect TileRegion = CurrentTileSelection.GetVisibleRegion();
 					FIntPoint NumTiles = OutFrame->Info.NumTiles / MipLevelDiv;
 					int32 PixelSize = sizeof(uint16) * OutFrame->Info.NumChannels;
 					TSharedPtr<FSampleConverterParameters> ConverterParams = MakeShared<FSampleConverterParameters>();
diff --git a/Engine/Plugins/Media/ImgMedia/Source/ImgMedia/Private/Readers/ExrImgMediaReaderGpu.cpp b/Engine/Plugins/Media/ImgMedia/Source/ImgMedia/Private/Readers/ExrImgMediaReaderGpu.cpp
--- a/Engine/Plugins/Media/ImgMedia/Source/ImgMedia/Private/Readers/ExrImgMediaReaderGpu.cpp
+++ b/Engine/Plugins/Media/ImgMedia/Source/ImgMedia/Private/Readers/ExrImgMediaReaderGpu.cpp
@@ -180,9 +180,10 @@ bool FExrImgMediaReaderGpu::ReadFrame(int32 FrameId, const TMap<int32, FImgMedia
 
 			if (ReadThisMip)
 			{
+				FIntRect TileRegion = CurrentTileSelection.GetVisibleRegion();
 				FIntRect& Viewport = ConverterParams->Viewports.Add(CurrentMipLevel);
-				Viewport.Min = FIntPoint(ConverterParams->TileDimWithBorders.X * CurrentTileSelection.TopLeftX, ConverterParams->TileDimWithBorders.Y * CurrentTileSelection.TopLeftY);
-				Viewport.Max = FIntPoint(ConverterParams->TileDimWithBorders.X * CurrentTileSelection.BottomRightX, ConverterParams->TileDimWithBorders.Y * CurrentTileSelection.BottomRightY);
+				Viewport.Min = FIntPoint(ConverterParams->TileDimWithBorders.X * TileRegion.Min.X, ConverterParams->TileDimWithBorders.Y * TileRegion.Min.Y);
+				Viewport.Max = FIntPoint(ConverterParams->TileDimWithBorders.X * TileRegion.Max.X, ConverterParams->TileDimWithBorders.Y * TileRegion.Max.Y);
 				Viewport.Clip(FIntRect(FIntPoint::ZeroValue, CurrentMipDim));
 
 				const SIZE_T BufferSize = GetBufferSize(CurrentMipDim, ConverterParams->FrameInfo.NumChannels, bHasTiles, OutFrame->Info.NumTiles / MipLevelDiv, ConverterParams->bCustomExr);
@@ -201,12 +202,6 @@ bool FExrImgMediaReaderGpu::ReadFrame(int32 FrameId, const TMap<int32, FImgMedia
 					// read frame data
 					if (bHasTiles || ConverterParams->bCustomExr)
 					{
-						FIntRect TileRegion = FIntRect(
-							(int32)CurrentTileSelection.TopLeftX,
-							(int32)CurrentTileSelection.TopLeftY,
-							(int32)CurrentTileSelection.BottomRightX,
-							(int32)CurrentTileSelection.BottomRightY);
-
 						ReadResult = ReadTilesCustom(MipDataPtr, ImagePath, FrameId, TileRegion, ConverterParams, CurrentMipLevel);
 					}
 					else
diff --git a/Engine/Plugins/Media/ImgMedia/Source/ImgMedia/Public/ImgMediaMipMapInfo.h b/Engine/Plugins/Media/ImgMedia/Source/ImgMedia/Public/ImgMediaMipMapInfo.h
--- a/Engine/Plugins/Media/ImgMedia/Source/ImgMedia/Public/ImgMediaMipMapInfo.h
+++ b/Engine/Plugins/Media/ImgMedia/Source/ImgMedia/Public/ImgMediaMipMapInfo.h
@@ -7,6 +7,7 @@
 #include "Tickable.h"
 #include "UObject/WeakObjectPtr.h"
 #include "UObject/WeakObjectPtrTemplates.h"
+#include "Containers/BitArray.h"
 
 class AActor;
 class UMeshComponent;
@@ -58,111 +59,118 @@ struct FImgMediaMipMapCameraInfo
  */
 struct FImgMediaTileSelection
 {
-	/** X position of the most top left visibile tile. */
-	uint16 TopLeftX;
-	/** Y position of the most top left visible tile. */
-	uint16 TopLeftY;
 	/**
-	 * X position of the most bottom right visible tile + 1.
-	 * If this is equal to TopLeftX, then this implies that the tile is not visible.
+	 * Create and initialize a new tile selection.
+	 *
+	 * @param NumTilesX Horizontal tile count.
+	 * @param NumTilesY Vertical tile count.
+	 * @param bDefaultVisibility Optional visibility for the entire region.
 	 */
-	uint16 BottomRightX;
-	/** Y position of the most bottom right visible tile + 1. */
-	uint16 BottomRightY;
-
-	FImgMediaTileSelection()
-	{
-		SetAllNotVisible();
-	}
+	FImgMediaTileSelection(int32 NumTilesX, int32 NumTilesY, bool bDefaultVisibility = false);
 
 	/**
-	 * Marks all tiles as visible.
+	 * Create and initialize a new tile selection, adjusting tile counts for the specified higher mip level.
+	 *
+	 * @param BaseNumTilesX Horizontal tile count at mip level 0.
+	 * @param NumTilesY Vertical tile count at mip level 0.
+	 * @param TargetMipLevel Higher target mip level for the selection (usually 1 and above).
+	 * @param bDefaultVisibility Optional visibility for the entire region.
 	 */
-	void SetAllVisible()
-	{
-		TopLeftX = 0;
-		TopLeftY = 0;
-		BottomRightX = 0xffff;
-		BottomRightY = 0xffff;
-	}
+	static FImgMediaTileSelection CreateForTargetMipLevel(int32 BaseNumTilesX, int32 BaseNumTilesY, int32 TargetMipLevel, bool bDefaultVisibility = false);
 
-	/**
-	 * Marks all tiles as not visible.
-	 */
-	void SetAllNotVisible()
-	{
-		TopLeftX = 0xffff;
-		TopLeftY = 0xffff;
-		BottomRightX = 0;
-		BottomRightY = 0;
-	}
+	FImgMediaTileSelection() = default;
+	~FImgMediaTileSelection() noexcept = default;
 
-	/**
-	 * Marks visible tile region.
-	 * 
-	 * @param InTopLeftX Top-left coordinate horizontal component.
-	 * @param InTopLeftY Top-left coordinate vertical component.
-	 * @param InBottomRightX Bottom-right coordinate horizontal component.
-	 * @param InBottomRightY Bottom-rightcoordinate vertical component.
-	 */
-	void SetVisibleRegion(uint16 InTopLeftX, uint16 InTopLeftY, uint16 InBottomRightX, uint16 InBottomRightY)
-	{
-		TopLeftX = InTopLeftX;
-		TopLeftY = InTopLeftY;
-		BottomRightX = InBottomRightX;
-		BottomRightY = InBottomRightY;
-	}
+	FImgMediaTileSelection(const FImgMediaTileSelection&) = default;
+	FImgMediaTileSelection& operator=(const FImgMediaTileSelection&) = default;
+
+	FImgMediaTileSelection(FImgMediaTileSelection&&) = default;
+	FImgMediaTileSelection& operator=(FImgMediaTileSelection&&) = default;
 
 	/**
-	* See if this selection is visible.
+	* See if any tile is visible.
 	*
 	* @return True if so.
 	*/
-	bool IsVisible() const { return (TopLeftX < BottomRightX); };
-
+	bool IsAnyVisible() const;
 
 	/**
-	 * Include a given tile coordinate to the current section region.
+	 * Check if a specific tile is visible.
 	 *
 	 * @param TileCoordX Horizontal tile coordinate.
 	 * @param TileCoordY Vertical tile coordinate.
+	 * @return True if the coordinate is visible, false otherwise.
 	 */
-	void Include(uint16 TileCoordX, uint16 TileCoordY)
-	{
-		TopLeftX = FMath::Min(TopLeftX, TileCoordX);
-		TopLeftY = FMath::Min(TopLeftY, TileCoordY);
-		BottomRightX = FMath::Max(BottomRightX, uint16(TileCoordX + 1u));
-		BottomRightY = FMath::Max(BottomRightY, uint16(TileCoordY + 1u));
-	}
+	bool IsVisible(int32 TileCoordX, int32 TileCoordY) const;
+
+	/**
+	 * Check if the currently visible tiles are also visible in another selection.
+	 *
+	 * @param Other Selection to compare.
+	 * @return True if the other selection is visible, false otherwise.
+	 */
+	bool Contains(const FImgMediaTileSelection& Other) const;
 
 	/**
-	 * Check if the current selection contains a tile.
+	 * Mark a tile as visible.
 	 *
 	 * @param TileCoordX Horizontal tile coordinate.
 	 * @param TileCoordY Vertical tile coordinate.
-	 * @return True if the coordinate is contained withing bounds, false otherwise.
 	 */
-	bool Contains(uint16 TileCoordX, uint16 TileCoordY) const
-	{
-		return TopLeftX <= TileCoordX &&
-			TopLeftY <= TileCoordY &&
-			BottomRightX > TileCoordX &&
-			BottomRightY > TileCoordY;
-	}
+	void SetVisible(int32 TileCoordX, int32 TileCoordY);
 
 	/**
-	 * Check if the current selection contains another selection within its bounds.
+	 * Returns the list of visible tile coordinates, in row and column order.
 	 *
-	 * @param Other Selection to compare.
-	 * @return True if the other selection is contained withing bounds, false otherwise.
+	 * @return Visible tile coordinates array.
+	 */
+	TArray<FIntPoint> GetVisibleCoordinates() const;
+
+	/**
+	 * Return the rectangular region bounding the visible tiles.
+	 *
+	 * @return Visible tile coordinates array.
+	 */
+	FIntRect GetVisibleRegion() const;
+
+	/**
+	 * Returns the number of visible tiles.
+	 *
+	 * @return Visible tile count.
 	 */
-	bool Contains(const FImgMediaTileSelection& Other) const
+	int32 NumVisibleTiles() const;
+
+	/**
+	 * Returns the overall dimensions in number of tiles.
+	 *
+	 * @return FIntPoint where X is the horizontal tile count and Y the vertical tile count.
+	 */
+	FIntPoint GetDimensions() const { return Dimensions; }
+
+private:
+	/** Convert from tile coordinates to linear index. */
+	FORCEINLINE static int32 ToIndex(int32 CoordX, int32 CoordY, const FIntPoint& Dim)
+	{
+		return CoordY * Dim.X + CoordX;
+	}
+
+	/** Convert from linear index to tile coordinates. */
+	FORCEINLINE static FIntPoint FromIndex(int32 Index, const FIntPoint& Dim)
 	{
-		return TopLeftX <= Other.TopLeftX &&
-			TopLeftY <= Other.TopLeftY &&
-			BottomRightX >= Other.BottomRightX &&
-			BottomRightY >= Other.BottomRightY;
+		return FIntPoint(Index % Dim.X, Index / Dim.X);
 	}
+	
+	/** Bit array of tiles over the entire region's dimensions, true when visible. */
+	TBitArray<> Tiles;
+	
+	/** Number of tiles in each axis. */
+	FIntPoint Dimensions;
+
+	/** Cached visible tile region. */
+	mutable FIntRect CachedVisibleRegion;
+	
+	/** Flag indicating if the cached visible tile region needs to be recalculated. */
+	mutable bool bCachedVisibleRegionDirty;
 };
 
 /**
