ImgMedia: Process images can be run on multiple threads when not using a player.

#rnx

diff --git a/Engine/Plugins/Media/ImgMedia/Source/ImgMediaEditor/Private/Widgets/ImgMediaProcessImagesOptions.h b/Engine/Plugins/Media/ImgMedia/Source/ImgMediaEditor/Private/Widgets/ImgMediaProcessImagesOptions.h
--- a/Engine/Plugins/Media/ImgMedia/Source/ImgMediaEditor/Private/Widgets/ImgMediaProcessImagesOptions.h
+++ b/Engine/Plugins/Media/ImgMedia/Source/ImgMediaEditor/Private/Widgets/ImgMediaProcessImagesOptions.h
@@ -43,6 +43,10 @@ public:
 	UPROPERTY(EditAnywhere, Transient, Category = Tiles, meta = (ClampMin = "0.0"))
 	int32 TileBorder = 0;
 
+	/** Number of threads to use when processing. If 0 then this number is set automatically. */
+	UPROPERTY(EditAnywhere, Transient, Category = Processing, meta = (ClampMin = "0.0"))
+	int32 NumThreads = 0;
+
 	/** Number of tiles in the X direction. If 0, then there are no tiles. */
 	UPROPERTY(VisibleAnywhere, Transient, Category = Tiles)
 	int32 NumTilesX = 0;
diff --git a/Engine/Plugins/Media/ImgMedia/Source/ImgMediaEditor/Private/Widgets/SImgMediaProcessImages.cpp b/Engine/Plugins/Media/ImgMedia/Source/ImgMediaEditor/Private/Widgets/SImgMediaProcessImages.cpp
--- a/Engine/Plugins/Media/ImgMedia/Source/ImgMediaEditor/Private/Widgets/SImgMediaProcessImages.cpp
+++ b/Engine/Plugins/Media/ImgMedia/Source/ImgMediaEditor/Private/Widgets/SImgMediaProcessImages.cpp
@@ -234,17 +234,31 @@ void SImgMediaProcessImages::ProcessAllImages()
 		else
 		{
 			IImageWrapperModule& ImageWrapperModule = FModuleManager::LoadModuleChecked<IImageWrapperModule>("ImageWrapper");
-			TSharedPtr<IImageWrapper> ImageWrapper = ImageWrapperModule.CreateImageWrapper(ImageFormat);
-
+			
 			// ImageWrapper is always returning an alpha channel for RGB, so check if we really have one.
 			bool bHasAlphaChannel = HasAlphaChannel(Ext, FPaths::Combine(SequencePath, FoundFiles[0]));
 
+			// Get number of threads to use.
+			int32 NumThreads = Options->NumThreads;
+			if (NumThreads <= 0)
+			{
+				NumThreads = 8;
+			}
+
 			// Loop through all files.
 			int NumDone = 0;
 			int TotalNum = FoundFiles.Num();
+			std::atomic<int32> NumActive = 0;
 			TSharedPtr<SNotificationItem> LocalConfirmNotification = ConfirmNotification;
 			for (const FString& FileName : FoundFiles)
 			{
+				// Wait for threads to finish if we have to many.
+				while (NumActive >= NumThreads)
+				{
+					FPlatformProcess::Sleep(0.1f);
+				}
+				NumActive++;
+				
 				// Update notification with current status.
 				Async(EAsyncExecution::TaskGraphMainThread, [LocalConfirmNotification, NumDone, TotalNum]()
 				{
@@ -257,33 +271,44 @@ void SImgMediaProcessImages::ProcessAllImages()
 				});
 				NumDone++;
 
-				FString FullFileName = FPaths::Combine(SequencePath, FileName);
-
-				// Load image into buffer.
-				TArray64<uint8> InputBuffer;
-				if (!FFileHelper::LoadFileToArray(InputBuffer, *FullFileName))
+				TSharedPtr<IImageWrapper> ImageWrapper = ImageWrapperModule.CreateImageWrapper(ImageFormat);
+				Async(EAsyncExecution::Thread, [this, SequencePath, FileName, OutPath,
+					bUseCustomFormat, Ext, &NumActive,
+					ImageWrapper, InTileWidth, InTileHeight,
+					TileBorder, bEnableMips, bHasAlphaChannel]() mutable
 				{
-					UE_LOG(LogImgMediaEditor, Error, TEXT("Failed to load %s"), *FullFileName);
-					break;
-				}
-				if (!ImageWrapper.IsValid() || !ImageWrapper->SetCompressed(InputBuffer.GetData(), InputBuffer.Num()))
-				{
-					UE_LOG(LogImgMediaEditor, Error, TEXT("Failed to create image wrapper for %s"), *FullFileName);
-					break;
-				}
+					FString FullFileName = FPaths::Combine(SequencePath, FileName);
 
-				// Import this image.
-				FString Name = FPaths::Combine(OutPath, FileName);
-				if (bUseCustomFormat)
-				{
-					ProcessImageCustom(ImageWrapper, InTileWidth, InTileHeight, TileBorder,
-						bEnableMips, bHasAlphaChannel, Name);
-				}
-				else
-				{
-					Name = FPaths::ChangeExtension(Name, TEXT(""));
-					ProcessImage(ImageWrapper, InTileWidth, InTileHeight, Name, Ext);
-				}
+					// Load image into buffer.
+					TArray64<uint8> InputBuffer;
+					if (!FFileHelper::LoadFileToArray(InputBuffer, *FullFileName))
+					{
+						UE_LOG(LogImgMediaEditor, Error, TEXT("Failed to load %s"), *FullFileName);
+						NumActive--;
+						return;
+					}
+					if (!ImageWrapper.IsValid() || !ImageWrapper->SetCompressed(InputBuffer.GetData(), InputBuffer.Num()))
+					{
+						UE_LOG(LogImgMediaEditor, Error, TEXT("Failed to create image wrapper for %s"), *FullFileName);
+						NumActive--;
+						return;
+					}
+
+					// Import this image.
+					FString Name = FPaths::Combine(OutPath, FileName);
+					if (bUseCustomFormat)
+					{
+						ProcessImageCustom(ImageWrapper, InTileWidth, InTileHeight, TileBorder,
+							bEnableMips, bHasAlphaChannel, Name);
+
+					}
+					else
+					{
+						Name = FPaths::ChangeExtension(Name, TEXT(""));
+						ProcessImage(ImageWrapper, InTileWidth, InTileHeight, Name, Ext);
+					}
+					NumActive--;
+				});
 
 				// Do we want to cancel?
 				if (bIsCancelling)
@@ -291,6 +316,12 @@ void SImgMediaProcessImages::ProcessAllImages()
 					break;
 				}
 			}
+
+			// Wait for all our tasks to finish.
+			while (NumActive > 0)
+			{
+				FPlatformProcess::Sleep(0.2f);
+			}
 		}
 	}
 
