ImgMedia: Integrate varying screen percentage into mip level calculations.

#rnx

diff --git a/Engine/Plugins/Media/ImgMedia/Source/ImgMedia/Private/Assets/ImgMediaMipMapInfo.cpp b/Engine/Plugins/Media/ImgMedia/Source/ImgMedia/Private/Assets/ImgMediaMipMapInfo.cpp
--- a/Engine/Plugins/Media/ImgMedia/Source/ImgMedia/Private/Assets/ImgMediaMipMapInfo.cpp
+++ b/Engine/Plugins/Media/ImgMedia/Source/ImgMedia/Private/Assets/ImgMediaMipMapInfo.cpp
@@ -350,8 +350,8 @@ namespace {
 								
 								if (bValidLevel)
 								{
-									MipLevelRange[0] = FMath::Min(MipLevelRange[0], FMath::Min((int32)CalculatedLevel, MaxLevel));
-									MipLevelRange[1] = FMath::Max(MipLevelRange[1], FMath::Min(FMath::CeilToInt32(CalculatedLevel), MaxLevel));
+									MipLevelRange[0] = FMath::Min(MipLevelRange[0], FMath::Clamp((int32)CalculatedLevel, 0, MaxLevel));
+									MipLevelRange[1] = FMath::Max(MipLevelRange[1], FMath::Clamp(FMath::CeilToInt32(CalculatedLevel), 0, MaxLevel));
 									NumVisibleCorners++;
 								}
 							}
diff --git a/Engine/Plugins/Media/ImgMedia/Source/ImgMedia/Private/Assets/ImgMediaMipMapInfoManager.cpp b/Engine/Plugins/Media/ImgMedia/Source/ImgMedia/Private/Assets/ImgMediaMipMapInfoManager.cpp
--- a/Engine/Plugins/Media/ImgMedia/Source/ImgMedia/Private/Assets/ImgMediaMipMapInfoManager.cpp
+++ b/Engine/Plugins/Media/ImgMedia/Source/ImgMedia/Private/Assets/ImgMediaMipMapInfoManager.cpp
@@ -13,6 +13,7 @@
 #include "Engine/GameViewportClient.h"
 #include "Engine/LocalPlayer.h"
 #include "GameFramework/PlayerController.h"
+#include "LegacyScreenPercentageDriver.h"
 #include "MovieSceneCommonHelpers.h"
 #include "SceneView.h"
 
@@ -141,8 +142,32 @@ void FImgMediaMipMapInfoManager::UpdateCameraInfo()
 						// Did we get a view?
 						if (SceneView != nullptr)
 						{
+							float ResolutionFraction = 1.0f;
+
+							if (ViewFamily.EngineShowFlags.ScreenPercentage)
+							{
+								// Get global view fraction.
+								FStaticResolutionFractionHeuristic StaticHeuristic;
+
+#if WITH_EDITOR
+								if (FStaticResolutionFractionHeuristic::FUserSettings::EditorOverridePIESettings())
+								{
+									StaticHeuristic.Settings.PullEditorRenderingSettings(/* bIsRealTime = */ true);
+								}
+								else
+#endif
+								{
+									StaticHeuristic.Settings.PullRunTimeRenderingSettings();
+								}
+
+								StaticHeuristic.PullViewFamilyRenderingSettings(ViewFamily);
+								StaticHeuristic.DPIScale = GameViewportClient->GetDPIScale();
+
+								ResolutionFraction = StaticHeuristic.ResolveResolutionFraction();
+							}
+
 							// Add camera info for this scene view.
-							AddCameraInfo(Viewport, SceneView);
+							AddCameraInfo(Viewport, SceneView, ResolutionFraction);
 						}
 					}
 				}
@@ -177,7 +202,7 @@ void FImgMediaMipMapInfoManager::UpdateCameraInfo()
 						FSceneView* SceneView = ViewportClient->CalcSceneView(&ViewFamily);
 						
 						// Add camera info for this scene view.
-						AddCameraInfo(Viewport, SceneView);
+						AddCameraInfo(Viewport, SceneView, ViewportClient->GetPreviewScreenPercentage() / 100.0f);
 					}
 				}
 			}
@@ -186,14 +211,22 @@ void FImgMediaMipMapInfoManager::UpdateCameraInfo()
 #endif
 }
 
-void FImgMediaMipMapInfoManager::AddCameraInfo(FViewport* Viewport, FSceneView* SceneView)
+void FImgMediaMipMapInfoManager::AddCameraInfo(FViewport* Viewport, FSceneView* SceneView, float ResolutionFraction)
 {
 	FImgMediaMipMapCameraInfo Info;
 	Info.Location = SceneView->ViewMatrices.GetViewOrigin();
 	Info.ViewMatrix = SceneView->ViewMatrices.GetViewMatrix();
 	Info.ViewProjectionMatrix = SceneView->ViewMatrices.GetViewProjectionMatrix();
-	Info.ViewportRect = FIntRect(FIntPoint::ZeroValue, Viewport->GetSizeXY());
-	Info.MaterialTextureMipBias = SceneView->MaterialTextureMipBias;
+	Info.ViewportRect = FIntRect(FIntPoint::ZeroValue, Viewport->GetSizeXY()).Scale(ResolutionFraction);
+
+	if (SceneView->PrimaryScreenPercentageMethod == EPrimaryScreenPercentageMethod::TemporalUpscale)
+	{
+		static const auto CVarMinAutomaticViewMipBiasOffset = IConsoleManager::Get().FindTConsoleVariableDataFloat(TEXT("r.ViewTextureMipBias.Offset"));
+		static const auto CVarMinAutomaticViewMipBias = IConsoleManager::Get().FindTConsoleVariableDataFloat(TEXT("r.ViewTextureMipBias.Min"));
+
+		Info.MaterialTextureMipBias = -(FMath::Max(-FMath::Log2(ResolutionFraction), 0.0f)) + CVarMinAutomaticViewMipBiasOffset->GetValueOnGameThread();
+		Info.MaterialTextureMipBias = FMath::Max(Info.MaterialTextureMipBias, CVarMinAutomaticViewMipBias->GetValueOnGameThread());
+	}
 
 	CameraInfos.Emplace(MoveTemp(Info));
 }
diff --git a/Engine/Plugins/Media/ImgMedia/Source/ImgMedia/Private/Assets/ImgMediaMipMapInfoManager.h b/Engine/Plugins/Media/ImgMedia/Source/ImgMedia/Private/Assets/ImgMediaMipMapInfoManager.h
--- a/Engine/Plugins/Media/ImgMedia/Source/ImgMedia/Private/Assets/ImgMediaMipMapInfoManager.h
+++ b/Engine/Plugins/Media/ImgMedia/Source/ImgMedia/Private/Assets/ImgMediaMipMapInfoManager.h
@@ -17,7 +17,6 @@ class IMediaPlayer;
 class UCameraComponent;
 class UImgMediaPlaybackComponent;
 class UMediaTexture;
-class FViewport;
 
 /**
  * Helps manage FImgMediaMipMapInfo objects.
@@ -98,8 +97,9 @@ protected:
 	 * 
 	 * @param Viewport		The camera's viewport size rectangle.
 	 * @param SceneView		The scene view we want to add.
+	 * @param ResolutionFraction	The viewport resolution fraction used to adjust mip level calculations.
 	 */
-	void AddCameraInfo(FViewport* Viewport, FSceneView* SceneView);
+	void AddCameraInfo(FViewport* Viewport, FSceneView* SceneView, float ResolutionFraction);
 
 	/**
 	 * Called when the cvar changes.
