ImgMedia: Process images can handle tile sizes that do not exactly fit a mip level.

#rnx

diff --git a/Engine/Plugins/Media/ImgMedia/Source/ImgMediaEditor/Private/Widgets/SImgMediaProcessImages.cpp b/Engine/Plugins/Media/ImgMedia/Source/ImgMediaEditor/Private/Widgets/SImgMediaProcessImages.cpp
--- a/Engine/Plugins/Media/ImgMedia/Source/ImgMediaEditor/Private/Widgets/SImgMediaProcessImages.cpp
+++ b/Engine/Plugins/Media/ImgMedia/Source/ImgMediaEditor/Private/Widgets/SImgMediaProcessImages.cpp
@@ -452,10 +452,10 @@ void SImgMediaProcessImages::ProcessImageCustomRawData(TArray64<uint8>& RawData,
 	TRACE_CPUPROFILER_EVENT_SCOPE(SImgMediaProcessImages::ProcessImageCustomRawData);
 	int32 DestWidth = Width;
 	int32 DestHeight = Height;
-	int32 NumTilesX = InTileWidth > 0 ? Width / InTileWidth : 1;
-	int32 NumTilesY = InTileHeight > 0 ? Height / InTileHeight : 1;
-	int32 TileWidth = Width / NumTilesX;
-	int32 TileHeight = Height / NumTilesY;
+	int32 NumTilesX = InTileWidth > 0 ? (Width + InTileWidth - 1) / InTileWidth : 1;
+	int32 NumTilesY = InTileHeight > 0 ? (Height + InTileHeight - 1) / InTileHeight : 1;
+	int32 TileWidth = InTileWidth;
+	int32 TileHeight = InTileHeight;
 	int32 BytesPerPixel = RawData.Num() / (Width * Height);
 	int32 BytesPerPixelPerChannel = BitDepth / 8;
 	int32 NumChannels = BytesPerPixel / BytesPerPixelPerChannel;
@@ -609,8 +609,8 @@ void SImgMediaProcessImages::ProcessImageCustomRawData(TArray64<uint8>& RawData,
 
 			int32 OutputWidth = 0;
 			int32 OutputHeight = 0;
-			int32 MipNumTilesX = MipSourceWidth / MipTileWidth;
-			int32 MipNumTilesY = MipSourceHeight / MipTileHeight;
+			int32 MipNumTilesX = (MipSourceWidth + MipTileWidth - 1) / MipTileWidth;
+			int32 MipNumTilesY = (MipSourceHeight + MipTileHeight - 1) / MipTileHeight;
 
 			// Make sure our sizes match the mip size we get from EXR.
 			int32 ExpectedMipWidth = MipSourceWidth + MipNumTilesX * InTileBorder * 2;
@@ -730,6 +730,7 @@ void SImgMediaProcessImages::TileData(uint8* SourceData, TArray64<uint8>& DestAr
 
 	int32 BytesPerTile = TileWidth * TileHeight * BytesPerPixel;
 	int32 ByterPerDestTile = DestTileWidth * DestTileHeight * BytesPerPixel;
+	uint8* DestTile = DestData;
 
 	// Loop over y tiles.
 	for (int32 TileY = 0; TileY < NumTilesY; ++TileY)
@@ -740,11 +741,19 @@ void SImgMediaProcessImages::TileData(uint8* SourceData, TArray64<uint8>& DestAr
 			// Get address of the source and destination tiles.
 			uint8* SourceTile = SourceData +
 				(TileX * TileWidth + TileY * SourceWidth * TileHeight) * BytesPerPixel;
-			uint8* DestTile = DestData + (TileX + TileY * NumTilesX) * ByterPerDestTile;
-
+			
+			// If this tile is over the right edge of our image, then make this tile smaller
+			// so it does not exceed the image size.
 			int32 NumberOfPixelsToCopy = TileWidth;
+			int32 ThisDestTileWidth = DestTileWidth;
+			if ((TileX + 1) * TileWidth > SourceWidth)
+			{
+				NumberOfPixelsToCopy = SourceWidth - TileX * TileWidth;
+				ThisDestTileWidth -= TileWidth - NumberOfPixelsToCopy;
+			}
 
 			// Create a left border.
+			int32 DestTileOffset = 0;
 			if (TileX > 0)
 			{
 				NumberOfPixelsToCopy += InTileBorder;
@@ -754,7 +763,8 @@ void SImgMediaProcessImages::TileData(uint8* SourceData, TArray64<uint8>& DestAr
 			else
 			{
 				// Offset the destination as we are skipping this border as we have no data.
-				DestTile += InTileBorder * BytesPerPixel;
+				DestTileOffset = InTileBorder * BytesPerPixel;
+				DestTile += DestTileOffset;
 			}
 
 			// Create a right border.
@@ -763,8 +773,16 @@ void SImgMediaProcessImages::TileData(uint8* SourceData, TArray64<uint8>& DestAr
 				NumberOfPixelsToCopy += InTileBorder;
 			}
 
+			// If this tile is over the bottom edge of our image, then make this tile smaller
+			// so it does not exceed the image size.
+			int32 ThisDestTileHeight = DestTileHeight;
+			if ((TileY + 1) * TileHeight > SourceHeight)
+			{
+				ThisDestTileHeight -= SourceHeight - TileY * TileHeight;
+			}
+
 			// Loop over each row in the tile.
-			for (int32 Row = 0; Row < DestTileHeight; ++Row)
+			for (int32 Row = 0; Row < ThisDestTileHeight; ++Row)
 			{
 				// Make sure we don't go beyond the source data.
 				int32 SourceRow = Row - InTileBorder;
@@ -778,10 +796,14 @@ void SImgMediaProcessImages::TileData(uint8* SourceData, TArray64<uint8>& DestAr
 				}
 
 				uint8* SourceLine = SourceTile + SourceRow * SourceWidth * BytesPerPixel;
-				uint8* DestLine = DestTile + Row * DestTileWidth * BytesPerPixel;
+				uint8* DestLine = DestTile;
 
 				// Copy the main data.
 				FMemory::Memcpy(DestLine, SourceLine, NumberOfPixelsToCopy * BytesPerPixel);
+
+				// Increment our pointer to the next tile.
+				// We have to remove any DestTileOffset we applied earlier.
+				DestTile += ThisDestTileWidth * BytesPerPixel - DestTileOffset;
 			}
 		}
 	}
