ImgMedia: Process images can tint mip levels different colours.

#rnx

diff --git a/Engine/Plugins/Media/ImgMedia/Source/ImgMediaEditor/Private/Widgets/ImgMediaProcessImagesOptions.cpp b/Engine/Plugins/Media/ImgMedia/Source/ImgMediaEditor/Private/Widgets/ImgMediaProcessImagesOptions.cpp
--- a/Engine/Plugins/Media/ImgMedia/Source/ImgMediaEditor/Private/Widgets/ImgMediaProcessImagesOptions.cpp
+++ b/Engine/Plugins/Media/ImgMedia/Source/ImgMediaEditor/Private/Widgets/ImgMediaProcessImagesOptions.cpp
@@ -5,5 +5,11 @@
 UImgMediaProcessImagesOptions::UImgMediaProcessImagesOptions(const FObjectInitializer& ObjectInitializer)
 	: Super(ObjectInitializer)
 {
+	MipLevelTints.Add(FLinearColor::Red);
+	MipLevelTints.Add(FLinearColor::Green);
+	MipLevelTints.Add(FLinearColor::Blue);
+	MipLevelTints.Add(FLinearColor::Yellow);
+	MipLevelTints.Add(FLinearColor(0.0f, 1.0f, 1.0f));
+	MipLevelTints.Add(FLinearColor(1.0f, 0.0f, 1.0f));
 }
 
diff --git a/Engine/Plugins/Media/ImgMedia/Source/ImgMediaEditor/Private/Widgets/ImgMediaProcessImagesOptions.h b/Engine/Plugins/Media/ImgMedia/Source/ImgMediaEditor/Private/Widgets/ImgMediaProcessImagesOptions.h
--- a/Engine/Plugins/Media/ImgMedia/Source/ImgMediaEditor/Private/Widgets/ImgMediaProcessImagesOptions.h
+++ b/Engine/Plugins/Media/ImgMedia/Source/ImgMediaEditor/Private/Widgets/ImgMediaProcessImagesOptions.h
@@ -54,4 +54,12 @@ public:
 	/** Number of tiles in the Y direction. If 0, then there are no tiles. */
 	UPROPERTY(VisibleAnywhere, Transient, Category = Tiles)
 	int32 NumTilesY = 0;
+
+	/** Tint each mip level a different colour to help with debugging. */
+	UPROPERTY(EditAnywhere, Transient, Category = Debug)
+	bool bEnableMipLevelTint = false;
+
+	/** Colour to tint each mip level. */
+	UPROPERTY(EditAnywhere, Transient, Category = Debug)
+	TArray<FLinearColor> MipLevelTints;
 };
diff --git a/Engine/Plugins/Media/ImgMedia/Source/ImgMediaEditor/Private/Widgets/SImgMediaProcessImages.cpp b/Engine/Plugins/Media/ImgMedia/Source/ImgMediaEditor/Private/Widgets/SImgMediaProcessImages.cpp
--- a/Engine/Plugins/Media/ImgMedia/Source/ImgMediaEditor/Private/Widgets/SImgMediaProcessImages.cpp
+++ b/Engine/Plugins/Media/ImgMedia/Source/ImgMediaEditor/Private/Widgets/SImgMediaProcessImages.cpp
@@ -471,6 +471,7 @@ void SImgMediaProcessImages::ProcessImageCustomRawData(TArray64<uint8>& RawData,
 	}
 
 	TArray64<uint8> TileBuffer;
+	TArray64<uint8> TintBuffer;
 	bool bIsTiled = (NumTilesX > 1) || (NumTilesY > 1);
 	if (bIsTiled)
 	{
@@ -590,6 +591,15 @@ void SImgMediaProcessImages::ProcessImageCustomRawData(TArray64<uint8>& RawData,
 			}
 		}
 
+		// Tint mip levels?
+		if (Options->bEnableMipLevelTint)
+		{
+			TintData(CurrentBuffer, TintBuffer,
+				MipLevel, MipWidth, MipHeight,
+				NumChannels);
+			CurrentBuffer = TintBuffer.GetData();
+		}
+
 		// Do we need to tile this mip?
 		// Need to also check that this is actually a valid mip level.
 		if ((bIsTiled) && (MipSourceWidth > 0) && (MipSourceHeight > 0))
@@ -697,6 +707,43 @@ void SImgMediaProcessImages::RemoveAlphaChannel(TArray64<uint8>& Buffer)
 	Buffer.SetNum((BufferSize * 3) / 4, false);
 }
 
+void SImgMediaProcessImages::TintData(uint8* SourceData, TArray64<uint8>& DestArray,
+	int32 InMipLevel, int32 InWidth, int32 InHeight, int32 InNumChannels)
+{
+	TRACE_CPUPROFILER_EVENT_SCOPE(SImgMediaProcessImages::TintData);
+
+	// Get tint colour.
+	FLinearColor TintColor = FLinearColor::White;
+	if (Options->MipLevelTints.Num() > 0)
+	{
+		TintColor = Options->MipLevelTints[InMipLevel % Options->MipLevelTints.Num()];
+	}
+
+	// Set up destination buffer.
+	DestArray.Reset();
+	DestArray.AddUninitialized(InWidth * InHeight * InNumChannels * 2);
+	FFloat16* Buffer = (FFloat16*)SourceData;
+	FFloat16* OutBuffer = (FFloat16*)DestArray.GetData();
+		
+	// Tint buffer.
+	for (int32 PixelY = 0; PixelY < InHeight; ++PixelY)
+	{
+		for (int32 PixelX = 0; PixelX < InWidth; ++PixelX)
+		{
+			OutBuffer[0] = (Buffer[0].GetFloat() + TintColor.R) * 0.5f;
+			OutBuffer[1] = (Buffer[1].GetFloat() + TintColor.G) * 0.5f;
+			OutBuffer[2] = (Buffer[2].GetFloat() + TintColor.B) * 0.5f;
+			if (InNumChannels == 4)
+			{
+				OutBuffer[3] = Buffer[3];
+			}
+				
+			Buffer += InNumChannels;
+			OutBuffer += InNumChannels;
+		}
+	}
+}
+
 void SImgMediaProcessImages::TileData(uint8* SourceData, TArray64<uint8>& DestArray,
 	int32 SourceWidth, int32 SourceHeight, int32 DestWidth, int32 DestHeight,
 	int32 NumTilesX, int32 NumTilesY,
diff --git a/Engine/Plugins/Media/ImgMedia/Source/ImgMediaEditor/Private/Widgets/SImgMediaProcessImages.h b/Engine/Plugins/Media/ImgMedia/Source/ImgMediaEditor/Private/Widgets/SImgMediaProcessImages.h
--- a/Engine/Plugins/Media/ImgMedia/Source/ImgMediaEditor/Private/Widgets/SImgMediaProcessImages.h
+++ b/Engine/Plugins/Media/ImgMedia/Source/ImgMediaEditor/Private/Widgets/SImgMediaProcessImages.h
@@ -120,6 +120,20 @@ private:
 	 */
 	void RemoveAlphaChannel(TArray64<uint8>& Buffer);
 
+	/**
+	 * Tints a source buffer and outputs it to a destination.
+	 *
+	 * @param SourceData		Source image.
+	 * @param DestArray			Destination image.
+	 * @param InMipLevel		The mip level of this image.
+	 * @param InWidth			Width of source in pixels.
+	 * @param InHeight			Height of source in pixels.
+	 * @param InNumChannels		Number of channels in the image.
+	 */
+	void TintData(uint8* SourceData, TArray64<uint8>& DestArray, 
+		int32 InMipLevel, int32 InWidth, int32 InHeight,
+		int32 InNumChannels);
+
 	/**
 	 * Creates tiles from a source and outputs it to a destination.
 	 *
