ImgMedia: Scene view extension replaces ImgMediaMipMapInfoManager to collect camera information. This adds support for nDisplay and unifies view/camera information caching for all viewport clients (including UDisplayClusterViewportClient).

diff --git a/Engine/Plugins/Media/ImgMedia/Source/ImgMedia/Private/Assets/ImgMediaMipMapInfo.cpp b/Engine/Plugins/Media/ImgMedia/Source/ImgMedia/Private/Assets/ImgMediaMipMapInfo.cpp
--- a/Engine/Plugins/Media/ImgMedia/Source/ImgMedia/Private/Assets/ImgMediaMipMapInfo.cpp
+++ b/Engine/Plugins/Media/ImgMedia/Source/ImgMedia/Private/Assets/ImgMediaMipMapInfo.cpp
@@ -2,7 +2,7 @@
 
 #include "ImgMediaMipMapInfo.h"
 
-#include "ImgMediaMipMapInfoManager.h"
+#include "IImgMediaModule.h"
 #include "ImgMediaPrivate.h"
 
 #include "Async/Async.h"
@@ -19,6 +19,14 @@
 
 DECLARE_CYCLE_STAT(TEXT("ImgMedia MipMap Update Cache"), STAT_ImgMedia_MipMapUpdateCache, STATGROUP_Media);
 
+static TAutoConsoleVariable<bool> CVarImgMediaMipMapDebugEnable(
+	TEXT("ImgMedia.MipMapDebug"),
+	0,
+	TEXT("Display debug on mipmaps used by the ImgMedia plugin.\n")
+	TEXT("   0: off (default)\n")
+	TEXT("   1: on\n"),
+	ECVF_Default);
+
 FImgMediaTileSelection::FImgMediaTileSelection(int32 NumTilesX, int32 NumTilesY, bool bDefaultVisibility)
 	: Tiles(bDefaultVisibility, NumTilesX * NumTilesY)
 	, Dimensions(NumTilesX, NumTilesY)
@@ -585,7 +593,6 @@ void FImgMediaMipMapInfo::RemoveObject(AActor* InActor)
 void FImgMediaMipMapInfo::AddObjectsUsingThisMediaTexture(UMediaTexture* InMediaTexture)
 {
 	// Get objects using this texture.
-	FImgMediaMipMapInfoManager& InfoManager = FImgMediaMipMapInfoManager::Get();
 	FMediaTextureTracker& TextureTracker = FMediaTextureTracker::Get();
 	const TArray<TWeakPtr<FMediaTextureTrackerObject, ESPMode::ThreadSafe>>* ObjectInfos = TextureTracker.GetObjects(InMediaTexture);
 	if (ObjectInfos != nullptr)
@@ -667,15 +674,13 @@ void FImgMediaMipMapInfo::Tick(float DeltaTime)
 {
 	FScopeLock Lock(&InfoCriticalSection);
 
-	// Get global info.
-	FImgMediaMipMapInfoManager& InfoManager = FImgMediaMipMapInfoManager::Get();
-	CameraInfos = InfoManager.GetCameraInfo();
+	CameraInfos = IImgMediaModule::Get().GetCopyCameraInfos();
 
 	// Let the cache update this frame.
 	bIsCacheValid = false;
 
 	// Display debug?
-	if (InfoManager.IsDebugEnabled())
+	if (CVarImgMediaMipMapDebugEnable.GetValueOnGameThread())
 	{
 		if (GEngine != nullptr)
 		{
diff --git a/Engine/Plugins/Media/ImgMedia/Source/ImgMedia/Private/Assets/ImgMediaMipMapInfoManager.cpp b/Engine/Plugins/Media/ImgMedia/Source/ImgMedia/Private/Assets/ImgMediaMipMapInfoManager.cpp
deleted file mode 100644
diff --git a/Engine/Plugins/Media/ImgMedia/Source/ImgMedia/Private/Assets/ImgMediaMipMapInfoManager.h b/Engine/Plugins/Media/ImgMedia/Source/ImgMedia/Private/Assets/ImgMediaMipMapInfoManager.h
deleted file mode 100644
diff --git a/Engine/Plugins/Media/ImgMedia/Source/ImgMedia/Private/Assets/ImgMediaSource.cpp b/Engine/Plugins/Media/ImgMedia/Source/ImgMedia/Private/Assets/ImgMediaSource.cpp
--- a/Engine/Plugins/Media/ImgMedia/Source/ImgMedia/Private/Assets/ImgMediaSource.cpp
+++ b/Engine/Plugins/Media/ImgMedia/Source/ImgMedia/Private/Assets/ImgMediaSource.cpp
@@ -5,7 +5,6 @@
 #include "IImgMediaModule.h"
 #include "ImgMediaGlobalCache.h"
 #include "ImgMediaMipMapInfo.h"
-#include "ImgMediaMipMapInfoManager.h"
 #include "ImgMediaPrivate.h"
 
 #include "HAL/FileManager.h"
@@ -59,18 +58,6 @@ void UImgMediaSource::SetSequencePath(const FString& Path)
 }
 
 
-void UImgMediaSource::AddGlobalCamera(AActor* InActor)
-{
-	FImgMediaMipMapInfoManager::Get().AddCamera(InActor);
-}
-
-
-void UImgMediaSource::RemoveGlobalCamera(AActor* InActor)
-{
-	FImgMediaMipMapInfoManager::Get().RemoveCamera(InActor);
-}
-
-
 void UImgMediaSource::AddTargetObject(AActor* InActor, float Width)
 {
 	MipMapInfo->AddObject(InActor, Width, 0.0f);
diff --git a/Engine/Plugins/Media/ImgMedia/Source/ImgMedia/Private/ImgMediaModule.cpp b/Engine/Plugins/Media/ImgMedia/Source/ImgMedia/Private/ImgMediaModule.cpp
--- a/Engine/Plugins/Media/ImgMedia/Source/ImgMedia/Private/ImgMediaModule.cpp
+++ b/Engine/Plugins/Media/ImgMedia/Source/ImgMedia/Private/ImgMediaModule.cpp
@@ -5,10 +5,10 @@
 #include "IMediaClock.h"
 #include "IMediaModule.h"
 #include "Misc/QueuedThreadPool.h"
-#include "Modules/ModuleManager.h"
 
 #include "ImgMediaGlobalCache.h"
 #include "ImgMediaPlayer.h"
+#include "ImgMediaSceneViewExtension.h"
 #include "ImgMediaScheduler.h"
 #include "ImgMediaSource.h"
 #include "IImgMediaModule.h"
@@ -122,6 +122,16 @@ public:
 		return Player;
 	}
 
+	TArray<FImgMediaMipMapCameraInfo> GetCopyCameraInfos() const override
+	{
+		if (SceneViewExtension.IsValid())
+		{
+			return SceneViewExtension->GetCameraInfos();
+		}
+
+		return {};
+	}
+
 public:
 
 	//~ IModuleInterface interface
@@ -135,6 +145,9 @@ public:
 		{
 			UMediaSource::RegisterSpawnFromFileExtension(Ext, SpawnDelegate);
 		}
+
+		FCoreDelegates::OnPostEngineInit.AddRaw(this, &FImgMediaModule::OnPostEngineInit);
+		FCoreDelegates::OnEnginePreExit.AddRaw(this, &FImgMediaModule::OnEnginePreExit);
 	}
 
 	virtual void ShutdownModule() override
@@ -148,6 +161,9 @@ public:
 		Scheduler.Reset();
 		GlobalCache.Reset();
 
+		FCoreDelegates::OnEnginePreExit.RemoveAll(this);
+		FCoreDelegates::OnPostEngineInit.RemoveAll(this);
+
 #if USE_IMGMEDIA_DEALLOC_POOL
 		ImgMediaThreadPool.Reset();
 #endif
@@ -155,6 +171,16 @@ public:
 
 private:
 
+	void OnPostEngineInit()
+	{
+		SceneViewExtension = FSceneViewExtensions::NewExtension<FImgMediaSceneViewExtension>();
+	}
+
+	void OnEnginePreExit()
+	{
+		SceneViewExtension.Reset();
+	}
+
 	void InitScheduler()
 	{
 		// initialize scheduler
@@ -193,6 +219,9 @@ private:
 
 	TSharedPtr<FImgMediaScheduler, ESPMode::ThreadSafe> Scheduler;
 
+	/** Scene view extension used to track view/camera info. */
+	TSharedPtr<FImgMediaSceneViewExtension, ESPMode::ThreadSafe> SceneViewExtension;
+
 	/** List of file extensions that we support. */
 	const TArray<FString> FileExtensions =
 	{
diff --git a/Engine/Plugins/Media/ImgMedia/Source/ImgMedia/Private/ImgMediaSceneViewExtension.cpp b/Engine/Plugins/Media/ImgMedia/Source/ImgMedia/Private/ImgMediaSceneViewExtension.cpp
new file mode 100644
--- /dev/null
+++ b/Engine/Plugins/Media/ImgMedia/Source/ImgMedia/Private/ImgMediaSceneViewExtension.cpp
@@ -0,0 +1,62 @@
+// Copyright Epic Games, Inc. All Rights Reserved.
+
+#include "ImgMediaSceneViewExtension.h"
+
+#include "SceneView.h"
+
+
+FImgMediaSceneViewExtension::FImgMediaSceneViewExtension(const FAutoRegister& AutoReg)
+	: FSceneViewExtensionBase(AutoReg)
+	, CachedCameraInfos()
+	, LastFrameNumber(0)
+{
+}
+
+void FImgMediaSceneViewExtension::SetupViewFamily(FSceneViewFamily& InViewFamily)
+{
+}
+
+void FImgMediaSceneViewExtension::SetupView(FSceneViewFamily& InViewFamily, FSceneView& InView)
+{
+}
+
+void FImgMediaSceneViewExtension::BeginRenderViewFamily(FSceneViewFamily& InViewFamily)
+{
+	if (LastFrameNumber != InViewFamily.FrameNumber)
+	{
+		CachedCameraInfos.Reset();
+		LastFrameNumber = InViewFamily.FrameNumber;
+	}
+
+	float ResolutionFraction = InViewFamily.SecondaryViewFraction;
+
+	if (InViewFamily.GetScreenPercentageInterface())
+	{
+		ResolutionFraction *= InViewFamily.GetPrimaryResolutionFractionUpperBound();
+	}
+
+	static const auto CVarMinAutomaticViewMipBiasOffset = IConsoleManager::Get().FindTConsoleVariableDataFloat(TEXT("r.ViewTextureMipBias.Offset"));
+	static const auto CVarMinAutomaticViewMipBias = IConsoleManager::Get().FindTConsoleVariableDataFloat(TEXT("r.ViewTextureMipBias.Min"));
+
+	for (const FSceneView* View : InViewFamily.Views)
+	{
+		FImgMediaMipMapCameraInfo Info;
+		Info.Location = View->ViewMatrices.GetViewOrigin();
+		Info.ViewMatrix = View->ViewMatrices.GetViewMatrix();
+		Info.ViewProjectionMatrix = View->ViewMatrices.GetViewProjectionMatrix();
+		Info.ViewportRect = View->UnconstrainedViewRect.Scale(ResolutionFraction);
+
+		// View->MaterialTextureMipBias is only set later in rendering: we replicate the logic here.
+		if (View->PrimaryScreenPercentageMethod == EPrimaryScreenPercentageMethod::TemporalUpscale)
+		{
+			Info.MaterialTextureMipBias = -(FMath::Max(-FMath::Log2(ResolutionFraction), 0.0f)) + CVarMinAutomaticViewMipBiasOffset->GetValueOnGameThread();
+			Info.MaterialTextureMipBias = FMath::Max(Info.MaterialTextureMipBias, CVarMinAutomaticViewMipBias->GetValueOnGameThread());
+		}
+		else
+		{
+			Info.MaterialTextureMipBias = 0.0f;
+		}
+
+		CachedCameraInfos.Add(MoveTemp(Info));
+	}
+}
\ No newline at end of file
diff --git a/Engine/Plugins/Media/ImgMedia/Source/ImgMedia/Private/ImgMediaSceneViewExtension.h b/Engine/Plugins/Media/ImgMedia/Source/ImgMedia/Private/ImgMediaSceneViewExtension.h
new file mode 100644
--- /dev/null
+++ b/Engine/Plugins/Media/ImgMedia/Source/ImgMedia/Private/ImgMediaSceneViewExtension.h
@@ -0,0 +1,32 @@
+// Copyright Epic Games, Inc. All Rights Reserved.
+
+#pragma once
+
+#include "CoreMinimal.h"
+#include "SceneViewExtension.h"
+
+#include "ImgMediaMipMapInfo.h"
+
+class FImgMediaSceneViewExtension final : public FSceneViewExtensionBase
+{
+public:
+	FImgMediaSceneViewExtension(const FAutoRegister& AutoReg);
+
+	/**
+	 * Get the cached camera information array, updated on the game thread by BeginRenderViewFamily.
+	 *
+	 * @return Array of info on each camera.
+	 */
+	const TArray<FImgMediaMipMapCameraInfo>& GetCameraInfos() const { return CachedCameraInfos; };
+
+	void SetupViewFamily(FSceneViewFamily& InViewFamily) override;
+	void SetupView(FSceneViewFamily& InViewFamily, FSceneView& InView) override;
+	void BeginRenderViewFamily(FSceneViewFamily& InViewFamily) override;
+
+private:
+	/** Array of info on each camera used for mipmap calculations. */
+	TArray<FImgMediaMipMapCameraInfo> CachedCameraInfos;
+
+	/** Last received FSceneViewFamily frame number. */
+	uint32 LastFrameNumber;
+};
diff --git a/Engine/Plugins/Media/ImgMedia/Source/ImgMedia/Private/Player/ImgMediaPlayer.cpp b/Engine/Plugins/Media/ImgMedia/Source/ImgMedia/Private/Player/ImgMediaPlayer.cpp
--- a/Engine/Plugins/Media/ImgMedia/Source/ImgMedia/Private/Player/ImgMediaPlayer.cpp
+++ b/Engine/Plugins/Media/ImgMedia/Source/ImgMedia/Private/Player/ImgMediaPlayer.cpp
@@ -14,11 +14,15 @@
 
 #include "ImgMediaLoader.h"
 #include "ImgMediaMipMapInfo.h"
-#include "ImgMediaMipMapInfoManager.h"
 #include "ImgMediaScheduler.h"
 #include "ImgMediaSettings.h"
 #include "ImgMediaTextureSample.h"
 
+#include "MediaPlayer.h"
+#include "MediaPlayerFacade.h"
+#include "MediaTexture.h"
+#include "MediaTextureTracker.h"
+
 #define IMG_MEDIA_PLAYER_VERSION 2
 #define LOCTEXT_NAMESPACE "FImgMediaPlayer"
 
@@ -202,12 +206,29 @@ bool FImgMediaPlayer::Open(const FString& Url, const IMediaOptions* Options)
 			{
 				// Tell mipmipinfo about our media textures.
 				MipMapInfo->ClearAllObjects();
-				FImgMediaMipMapInfoManager& MipMapInfoManager = FImgMediaMipMapInfoManager::Get();
-				TArray<UMediaTexture*> MediaTextures;
-				MipMapInfoManager.GetMediaTexturesFromPlayer(MediaTextures, this);
-				for (UMediaTexture* MediaTexture : MediaTextures)
+
+				FMediaTextureTracker& TextureTracker = FMediaTextureTracker::Get();
+
+				// Look through all the media textures we know about.
+				for (TWeakObjectPtr<UMediaTexture> TexturePtr : TextureTracker.GetTextures())
 				{
-					MipMapInfo->AddObjectsUsingThisMediaTexture(MediaTexture);
+					UMediaTexture* Texture = TexturePtr.Get();
+					if (Texture != nullptr)
+					{
+						// Does this match the player?
+						UMediaPlayer* MediaPlayer = Texture->GetMediaPlayer();
+						if (MediaPlayer != nullptr)
+						{
+							TSharedPtr<IMediaPlayer, ESPMode::ThreadSafe> Player = MediaPlayer->GetPlayerFacade()->GetPlayer();
+							if (Player.IsValid())
+							{
+								if (Player.Get() == this)
+								{
+									MipMapInfo->AddObjectsUsingThisMediaTexture(Texture);
+								}
+							}
+						}
+					}
 				}
 			}
 		}
diff --git a/Engine/Plugins/Media/ImgMedia/Source/ImgMedia/Public/IImgMediaModule.h b/Engine/Plugins/Media/ImgMedia/Source/ImgMedia/Public/IImgMediaModule.h
--- a/Engine/Plugins/Media/ImgMedia/Source/ImgMedia/Public/IImgMediaModule.h
+++ b/Engine/Plugins/Media/ImgMedia/Source/ImgMedia/Public/IImgMediaModule.h
@@ -7,8 +7,10 @@
 #include "ProfilingDebugging/CsvProfiler.h"
 #include "Templates/SharedPointer.h"
 #include "Modules/ModuleInterface.h"
+#include "Modules/ModuleManager.h"
 
 class FImgMediaGlobalCache;
+struct FImgMediaMipMapCameraInfo;
 class FImgMediaPlayer;
 class IMediaEventSink;
 class IMediaPlayer;
@@ -26,6 +28,23 @@ class IMGMEDIA_API IImgMediaModule
 {
 public:
 
+	/**
+	* Singleton-like access to this module's interface.  This is just for convenience!
+	* Beware of calling this during the shutdown phase, though.  Your module might have been unloaded already.
+	*
+	* @return Returns singleton instance, loading the module on demand if needed
+	*/
+	static inline IImgMediaModule& Get()
+	{
+		static IImgMediaModule* ImgMediaModulePtr = nullptr;
+		if (!ImgMediaModulePtr)
+		{
+			ImgMediaModulePtr = FModuleManager::LoadModulePtr<IImgMediaModule>("ImgMedia");
+		}
+
+		return *ImgMediaModulePtr;
+	}
+
 	/**
 	 * Creates a media player for image sequences.
 	 *
@@ -34,6 +53,12 @@ public:
 	 */
 	virtual TSharedPtr<IMediaPlayer, ESPMode::ThreadSafe> CreatePlayer(IMediaEventSink& EventHandler) = 0;
 
+	/**
+	 * Get a copy of camera information cached by our local scene view extension, for mip & tile calculations.
+	 *
+	 * @return Array of info on each camera.
+	 */
+	virtual TArray<FImgMediaMipMapCameraInfo> GetCopyCameraInfos() const = 0;
 public:
 
 	/** Virtual destructor. */
diff --git a/Engine/Plugins/Media/ImgMedia/Source/ImgMedia/Public/ImgMediaMipMapInfo.h b/Engine/Plugins/Media/ImgMedia/Source/ImgMedia/Public/ImgMediaMipMapInfo.h
--- a/Engine/Plugins/Media/ImgMedia/Source/ImgMedia/Public/ImgMediaMipMapInfo.h
+++ b/Engine/Plugins/Media/ImgMedia/Source/ImgMedia/Public/ImgMediaMipMapInfo.h
@@ -262,13 +262,6 @@ public:
 	 */
 	TMap<int32, FImgMediaTileSelection> GetVisibleTiles();
 	
-	/**
-	 * Get information on all our cameras.
-	 *
-	 * @return Array of cameras.
-	 */
-	const TArray<FImgMediaMipMapCameraInfo>& GetCameraInfo() const { return CameraInfos; }
-
 	/**
 	 * Get information on objects that are using our textures.
 	 * 
diff --git a/Engine/Plugins/Media/ImgMedia/Source/ImgMedia/Public/ImgMediaSource.h b/Engine/Plugins/Media/ImgMedia/Source/ImgMedia/Public/ImgMediaSource.h
--- a/Engine/Plugins/Media/ImgMedia/Source/ImgMedia/Public/ImgMediaSource.h
+++ b/Engine/Plugins/Media/ImgMedia/Source/ImgMedia/Public/ImgMediaSource.h
@@ -110,22 +110,6 @@ public:
 	UFUNCTION(BlueprintCallable, Category="ImgMedia|ImgMediaSource")
 	void SetSequencePath(const FString& Path);
 
-	/**
-	 * This camera could be looking at any img sequence.
-	 *
-	 * @param InActor Camera object.
-	 */
-	UFUNCTION(BlueprintCallable, Category = "ImgMedia|ImgMediaSource")
-	void AddGlobalCamera(AActor* InActor);
-
-	/**
-	 * This camera is no longer looking at any img seqeunces.
-	 *
-	 * @param InActor Camera Object.
-	 */
-	UFUNCTION(BlueprintCallable, Category = "ImgMedia|ImgMediaSource")
-	void RemoveGlobalCamera(AActor* InActor);
-
 	/**
 	 * This object is using our img sequence.
 	 *
