ImgMediaReader ExrReaderGPU: Part1 Added ability to read the exact bytes of incomplete tiles for both custom and vanilla exr.

diff --git a/Engine/Plugins/Media/ImgMedia/Source/ExrReaderGpu/Private/ExrReaderGpu.cpp b/Engine/Plugins/Media/ImgMedia/Source/ExrReaderGpu/Private/ExrReaderGpu.cpp
--- a/Engine/Plugins/Media/ImgMedia/Source/ExrReaderGpu/Private/ExrReaderGpu.cpp
+++ b/Engine/Plugins/Media/ImgMedia/Source/ExrReaderGpu/Private/ExrReaderGpu.cpp
@@ -251,8 +251,107 @@ bool FExrReader::GenerateTextureData(uint16* Buffer, int32 BufferSize, FString F
 	return true;
 }
 
-bool FExrReader::OpenExrAndPrepareForPixelReading(FString FilePath, int32 NumOffsets, int32 NumLevels, bool bCustomExr)
+
+void FExrReader::CalculateTileOffsets
+	( TArray<int32>& OutNumTilesPerLevel
+	, TArray<TArray<int64>>& OutCustomOffsets
+	, const FIntPoint& FullTextureResolution
+	, const FIntPoint& FullResolutionInTiles
+	, const FIntPoint& TileDimWithBorders
+	, int32 NumMipLevels
+	, int64 PixelSize
+	, bool bCustomExr)
+{
+	TRACE_CPUPROFILER_EVENT_SCOPE_TEXT(*FString::Printf(TEXT("CalculateTileOffsets")));
+
+	OutCustomOffsets.SetNum(NumMipLevels);
+	// Custom exr files don't have tile offsets, therefore we need to calculate these manually.
+	if (bCustomExr)
+	{
+		int64 CurrentPosition = 0;
+		for (int32 MipLevel = 0; MipLevel < NumMipLevels; MipLevel++)
+		{
+			int32 MipDiv = 1 << MipLevel;
+
+			// Dimension of the texture in tiles including partial tiles
+			FIntPoint DimensionInTiles_PartialTiles
+				( FMath::CeilToInt(float(FullResolutionInTiles.X) / MipDiv)
+				, FMath::CeilToInt(float(FullResolutionInTiles.Y) / MipDiv));
+
+			// Dimension of the texture in tiles excluding partial tiles.
+			FIntPoint DimensionInTiles_CompleteTiles
+				( FMath::FloorToInt(float(FullResolutionInTiles.X) / MipDiv)
+				, FMath::FloorToInt(float(FullResolutionInTiles.Y) / MipDiv));
+
+			// Total number of tiles for this mip level.
+			int32 NumActualTiles = DimensionInTiles_PartialTiles.X * DimensionInTiles_PartialTiles.Y;
+
+			// Custom exr has one Vanilla tile per mip level that contains custom tiles.
+			OutNumTilesPerLevel.Add(1);
+
+			// Resolution of the texture in pixels for this mip level.
+			FIntPoint MipResolution = FullTextureResolution / MipDiv;
+
+			// Resolution of the partial tile in the bottom right corner.
+			FIntPoint PartialTileResolution = FIntPoint(MipResolution.X % TileDimWithBorders.X, MipResolution.Y % TileDimWithBorders.Y);
+
+			bool bHasPartialTiles = DimensionInTiles_PartialTiles != DimensionInTiles_CompleteTiles;
+
+			TArray<int64>& CurrentMipOffsets = OutCustomOffsets[MipLevel];
+			CurrentMipOffsets.SetNum(NumActualTiles);
+			for (int TileIndex = 0; TileIndex < NumActualTiles; TileIndex++)
+			{
+				if (TileIndex == 0)
+				{
+					CurrentPosition += FExrReader::TILE_PADDING;
+				}
+
+				if (MipLevel == 0 && TileIndex == 0)
+				{
+					CurrentMipOffsets[TileIndex] = CurrentPosition;
+					continue;
+				}
+
+				FIntPoint TileDim(TileDimWithBorders);
+
+				if (bHasPartialTiles)
+				{
+					// If true - this is a partial tile in X dimension.
+					if (((TileIndex) % DimensionInTiles_PartialTiles.X) == 0)
+					{
+						TileDim.X = PartialTileResolution.X;
+					}
+
+					// If true - this is a partial tile in Y dimension.
+					if (TileIndex > (DimensionInTiles_PartialTiles.X * DimensionInTiles_CompleteTiles.Y))
+					{
+						TileDim.Y = PartialTileResolution.Y;
+					}
+				}
+
+				// Tile offset.
+				CurrentPosition += TileDim.X * TileDim.Y * PixelSize;
+
+				// First tile at each mip level contains padding.
+				CurrentMipOffsets[TileIndex] = CurrentPosition;
+			}
+		}
+	}
+	else
+	{
+		for (int32 MipLevel = 0; MipLevel < NumMipLevels; MipLevel++)
+		{
+			int32 MipDiv = 1 << MipLevel;
+			int32 DimX = FMath::Max(1, FMath::CeilToInt(float(FullResolutionInTiles.X) / MipDiv));
+			int32 DimY = FMath::Max(1, FMath::CeilToInt(float(FullResolutionInTiles.Y) / MipDiv));
+			OutNumTilesPerLevel.Add(DimX * DimY);
+		}
+	}
+}
+
+bool FExrReader::OpenExrAndPrepareForPixelReading(FString FilePath, const TArray<int32>& NumOffsetsPerLevel, TArray<TArray<int64>>&& CustomOffsets, bool bInCustomExr)
 {
+ 	bCustomExr = bInCustomExr;
 	if (FileHandle != nullptr)
 	{
 		UE_LOG(LogExrReaderGpu, Error, TEXT("The file has already been open for reading but never closed."));
@@ -266,6 +365,7 @@ bool FExrReader::OpenExrAndPrepareForPixelReading(FString FilePath, int32 NumOff
 		return false;
 	}
 
+
 	// Reading header (throwaway) and line offset data. We just need line offset data.
 	{
 		ReadMagicNumberAndVersionField(FileHandle);
@@ -273,22 +373,33 @@ bool FExrReader::OpenExrAndPrepareForPixelReading(FString FilePath, int32 NumOff
 		{
 			return false;
 		}
-		LineOrTileOffsetsPerLevel.SetNum(NumLevels);
+
+		LineOrTileOffsetsPerLevel.SetNum(NumOffsetsPerLevel.Num());
 
 		// At the moment we support only increasing Y.
 		ELineOrder LineOrder = INCREASING_Y;
-		for (int Level = 0; Level < NumLevels; Level++)
+		for (int Level = 0; Level < LineOrTileOffsetsPerLevel.Num(); Level++)
+		{
+			TArray<int64>& OffsetsPerLevel = LineOrTileOffsetsPerLevel[Level];
+			OffsetsPerLevel.SetNum(NumOffsetsPerLevel[Level]);
+			ReadLineOrTileOffsets(FileHandle, LineOrder, OffsetsPerLevel);
+		}
+
+		if (bCustomExr)
 		{
-			LineOrTileOffsetsPerLevel[Level].SetNum(bCustomExr ? 1 : NumOffsets >> (Level * 2));
-			ReadLineOrTileOffsets(FileHandle, LineOrder, LineOrTileOffsetsPerLevel[Level]);
+			PixelStartByteOffset = LineOrTileOffsetsPerLevel[0][0];
+			LineOrTileOffsetsPerLevel = MoveTemp(CustomOffsets);
 		}
 	}
 
+	fseek(FileHandle, 0, SEEK_END);
+	FileLength = ftell(FileHandle);
 	fseek(FileHandle, LineOrTileOffsetsPerLevel[0][0], SEEK_SET);
 
 	return true;
 }
 
+
 bool FExrReader::ReadExrImageChunk(void* Buffer, int64 ChunkSize)
 {
 	if (FileHandle == nullptr)
@@ -310,26 +421,46 @@ bool FExrReader::ReadExrImageChunk(void* Buffer, int64 ChunkSize)
 	return NumElementsRead == NumElements;
 }
 
-bool FExrReader::SeekTileWithinFile(const int32 StartTileIndex, const FIntPoint& TexDimInTiles, const int32 Level, int64& OutBufferOffset)
+bool FExrReader::SeekTileWithinFile(const int32 StartTileIndex, const int32 Level, int64& OutBufferOffset)
 {
 	if (StartTileIndex >= LineOrTileOffsetsPerLevel[Level].Num())
 	{
 		UE_LOG(LogExrReaderGpu, Error, TEXT("Tile index is invalid."));
 		return false;
 	}
+
 	int64 LineOffset = LineOrTileOffsetsPerLevel[Level][StartTileIndex];
 	OutBufferOffset = LineOffset - LineOrTileOffsetsPerLevel[Level][0];
-	return fseek(FileHandle, LineOffset, SEEK_SET) == 0;
+	return fseek(FileHandle, LineOffset + (bCustomExr ? PixelStartByteOffset : 0), SEEK_SET) == 0;
 }
 
-bool FExrReader::SeekTileWithinFileCustom(const int32 StartTileIndex, const int64 TileStride, const int32 Level, int64& OutBufferOffset)
+bool FExrReader::GetByteOffsetForTile(const int32 TileIndex, const int32 Level, int64& OutBufferOffset)
 {
-	// Our custom exr is written as one single tile. In Exr every tile starts with 20 byte padding.
-	// Therefore we want to ignore it, when we read custom tiles.
-	const int32 TilePadding = 20;
-	int64 LineOffset = TilePadding + LineOrTileOffsetsPerLevel[Level][0] + StartTileIndex * TileStride;
-	OutBufferOffset = StartTileIndex * TileStride;
-	return fseek(FileHandle, LineOffset, SEEK_SET) == 0;
+	if (LineOrTileOffsetsPerLevel.Num() <= Level || LineOrTileOffsetsPerLevel[Level].Num() < TileIndex)
+	{
+		return false;
+	}
+	else
+	{
+		// In some cases we'd like to know the position of the final byte of the last tile.
+		if (TileIndex == LineOrTileOffsetsPerLevel[Level].Num())
+		{
+			// If this is the last mip and the last tile that means that we don't have 
+			if (LineOrTileOffsetsPerLevel.Num() > Level + 1)
+			{
+				OutBufferOffset = LineOrTileOffsetsPerLevel[Level + 1][0] + (bCustomExr ? PixelStartByteOffset : 0);
+			}
+			else
+			{
+				OutBufferOffset = FileLength;
+			}
+		}
+		else
+		{
+			OutBufferOffset = LineOrTileOffsetsPerLevel[Level][TileIndex] + (bCustomExr ? PixelStartByteOffset : 0);
+		}
+		return true;
+	}
 }
 
 bool FExrReader::CloseExrFile()
diff --git a/Engine/Plugins/Media/ImgMedia/Source/ExrReaderGpu/Public/ExrReaderGpu.h b/Engine/Plugins/Media/ImgMedia/Source/ExrReaderGpu/Public/ExrReaderGpu.h
--- a/Engine/Plugins/Media/ImgMedia/Source/ExrReaderGpu/Public/ExrReaderGpu.h
+++ b/Engine/Plugins/Media/ImgMedia/Source/ExrReaderGpu/Public/ExrReaderGpu.h
@@ -39,7 +39,13 @@ public:
 	static const int32	TILE_PADDING = 20;
 
 public:
-	FExrReader() :FileHandle(nullptr) {};
+	FExrReader() 
+		: FileHandle(nullptr)
+		, FileLength(0)
+		, bCustomExr(false)
+		, LineOrTileOffsetsPerLevel()
+		, PixelStartByteOffset(0)
+	{};
 private:
 
 	static bool ReadMagicNumberAndVersionField(FILE* FileHandle);
@@ -59,12 +65,25 @@ public:
 	/** Discards the header and fills the buffer with plain pixel data. */
 	static bool GenerateTextureData(uint16* Buffer, int32 BufferSize, FString FilePath, int32 NumberOfScanlines, int32 NumChannels);
 
+	/**
+	 * To be able to read tiles, especially in cases where tiles vary in sizes, we need to calculate offsets manually
+	 * for custom exr files, since we don't write out offsets.
+	*/
+	static void CalculateTileOffsets(TArray<int32>& OutNumTilesPerLevel
+		, TArray<TArray<int64>>& OutCustomOffsets
+		, const FIntPoint& FullTextureResolution
+		, const FIntPoint& FullResolutionInTiles
+		, const FIntPoint& TileDimWithBorders
+		, int32 NumMipLevels
+		, int64 PixelSize
+		, bool bCustomExr);
+
 	/** 
 	* This function is used to open file and keep a handle to it so the file can be read in chunks.
 	* Reading in chunks allows the process of reading to be canceled midway through reading. 
-	* This function reads and discards the header and keeps the pointers to the scanlines. 
+	* This function reads and discards the header and keeps the pointers to the scanlines or tiles. 
 	*/
-	bool OpenExrAndPrepareForPixelReading(FString FilePath, int32 NumOffsets, int32 NumLevels, bool bCustomExr = false);
+	bool OpenExrAndPrepareForPixelReading(FString FilePath, const TArray<int32>& NumOffsetsPerLevel, TArray<TArray<int64>>&& CustomOffsets, bool bInCustomExr = false);
 
 	/** 
 	* Read a chunk of an Exr file previously open via OpenExrAndPrepareForPixelReading.
@@ -75,12 +94,12 @@ public:
 	/**
 	* Moves to the specified tile position in the file in preparation for reading.
 	*/
-	bool SeekTileWithinFile(const int32 StartTileIndex, const FIntPoint& TexDimInTiles, const int32 Level, int64& OutBufferOffset);
+	bool SeekTileWithinFile(const int32 StartTileIndex, const int32 Level, int64& OutBufferOffset);
 
-	/** 
-	* Moves to the specified tile position in the file in preparation for reading.
+	/**
+	* Gets byte location of the requested tile for the requested mip level.
 	*/
-	bool SeekTileWithinFileCustom(const int32 StartTileIndex, const int64 TileStride, const int32 Level, int64& OutBufferOffset);
+	bool GetByteOffsetForTile(const int32 TileIndex, const int32 Level, int64& OutBufferOffset);
 
 	/**
 	* After we have finished reading the file that was open via OpenExrAndPrepareForPixelReading
@@ -96,9 +115,24 @@ private:
 	*/
 	FILE* FileHandle;
 
+	/**
+	* The length of the file in question. Used to get the last tile.
+	*/
+	int64 FileLength;
+
+	/**
+	* Is this exr a custom exr with swizzled data.
+	*/
+	bool bCustomExr;
+
 	/**
 	* These are byte offsets pointing to scanlines or tiles from the begining of the file.
 	*/
 	TArray<TArray<int64>> LineOrTileOffsetsPerLevel;
+
+	/** 
+	* Byte offset of the first pixel in file. This is equal to the size of the header.
+	*/
+	int64 PixelStartByteOffset;
 };
 #endif
diff --git a/Engine/Plugins/Media/ImgMedia/Source/ImgMedia/Private/Readers/ExrImgMediaReader.cpp b/Engine/Plugins/Media/ImgMedia/Source/ImgMedia/Private/Readers/ExrImgMediaReader.cpp
--- a/Engine/Plugins/Media/ImgMedia/Source/ImgMedia/Private/Readers/ExrImgMediaReader.cpp
+++ b/Engine/Plugins/Media/ImgMedia/Source/ImgMedia/Private/Readers/ExrImgMediaReader.cpp
@@ -57,6 +57,7 @@ bool FExrImgMediaReader::GetFrameInfo(const FString& ImagePath, FImgMediaFrameIn
 
 FExrImgMediaReader::EReadResult FExrImgMediaReader::ReadTilesCustom
 	( uint16* Buffer
+	, int64 BufferSize
 	, const FString& ImagePath
 	, int32 FrameId
 	, const FIntRect& TileRegion
@@ -69,16 +70,33 @@ FExrImgMediaReader::EReadResult FExrImgMediaReader::ReadTilesCustom
 	int64 CurrentBufferPos = 0;
 	FExrReader ChunkReader;
 	int MipLevelDiv = 1 << CurrentMipLevel;
-	FIntPoint DimensionInTiles = ConverterParams->FrameInfo.NumTiles / MipLevelDiv;
-	int32 NumTiles = ConverterParams->bCustomExr ? 1 : ConverterParams->FrameInfo.NumTiles.X * ConverterParams->FrameInfo.NumTiles.Y;
+
+	FIntPoint DimensionInTiles
+		( FMath::CeilToInt(float(ConverterParams->FrameInfo.NumTiles.X) / MipLevelDiv)
+		, FMath::CeilToInt(float(ConverterParams->FrameInfo.NumTiles.Y) / MipLevelDiv));
+
+	TArray<int32> NumTilesPerLevel;
+	TArray<TArray<int64>> CustomOffsets;
 	int32 NumMipLevels = ConverterParams->bMipsInSeparateFiles ? 1 : ConverterParams->NumMipLevels;
-	if (!ChunkReader.OpenExrAndPrepareForPixelReading(ImagePath, NumTiles, NumMipLevels, ConverterParams->bCustomExr))
+
+	FExrReader::CalculateTileOffsets(
+		NumTilesPerLevel, 
+		CustomOffsets, 
+		ConverterParams->FullResolution, 
+		ConverterParams->FrameInfo.NumTiles,
+		ConverterParams->TileDimWithBorders,
+		NumMipLevels, 
+		ConverterParams->PixelSize,
+		ConverterParams->bCustomExr);
+
+	if (!ChunkReader.OpenExrAndPrepareForPixelReading(ImagePath, NumTilesPerLevel, MoveTemp(CustomOffsets), ConverterParams->bCustomExr))
 	{
 		return Fail;
 	}
 
 	for (int32 TileRow = TileRegion.Min.Y; TileRow < TileRegion.Max.Y; TileRow++)
 	{
+
 		// Check to see if the frame was canceled.
 		{
 			FScopeLock RegionScopeLock(&CanceledFramesCriticalSection);
@@ -89,23 +107,30 @@ FExrImgMediaReader::EReadResult FExrImgMediaReader::ReadTilesCustom
 				break;
 			}
 		}
+
 		const uint16 Padding = ConverterParams->bCustomExr ? 0 : FExrReader::TILE_PADDING;
 		const int64 TileByteStride = ConverterParams->PixelSize * ConverterParams->TileDimWithBorders.X * ConverterParams->TileDimWithBorders.Y + Padding;
 		const int StartTileIndex = TileRow * DimensionInTiles.X + TileRegion.Min.X;
+
+		bool bLastTile = TileRow + 1 == DimensionInTiles.Y && TileRegion.Max.X == DimensionInTiles.X;
 		const int EndTileIndex = TileRow * DimensionInTiles.X + TileRegion.Max.X;
-		const int BytesToRead = (EndTileIndex - StartTileIndex) * TileByteStride;
 
-		if (ConverterParams->bCustomExr)
-		{
-			const int64 TileStride = ConverterParams->TileDimWithBorders.X * ConverterParams->TileDimWithBorders.Y * ConverterParams->PixelSize;
-			ChunkReader.SeekTileWithinFileCustom(StartTileIndex, TileStride, ConverterParams->bMipsInSeparateFiles ? 0 : CurrentMipLevel, CurrentBufferPos);
-		}
-		else
+		int MipLevel = ConverterParams->bMipsInSeparateFiles ? 0 : CurrentMipLevel;
+		ChunkReader.SeekTileWithinFile(StartTileIndex, MipLevel, CurrentBufferPos);
+		int64 ByteOffsetStart = 0;
+		int64 ByteOffsetEnd = 0;
+
+		if (!ChunkReader.GetByteOffsetForTile(StartTileIndex, MipLevel, ByteOffsetStart)
+			|| !ChunkReader.GetByteOffsetForTile(EndTileIndex, MipLevel, ByteOffsetEnd))
 		{
-			ChunkReader.SeekTileWithinFile(StartTileIndex, DimensionInTiles, ConverterParams->bMipsInSeparateFiles ? 0 : CurrentMipLevel, CurrentBufferPos);
+			bResult = Fail;
+			break;
 		}
 
-		if (!ChunkReader.ReadExrImageChunk(reinterpret_cast<char*>(Buffer) + CurrentBufferPos, BytesToRead))
+		int64 ByteChunkToRead = bLastTile
+			? FMath::Min(ByteOffsetEnd - ByteOffsetStart, BufferSize - CurrentBufferPos)
+			: (ByteOffsetEnd - ByteOffsetStart);
+		if (!ChunkReader.ReadExrImageChunk(reinterpret_cast<char*>(Buffer) + CurrentBufferPos, ByteChunkToRead))
 		{
 			bResult = Fail;
 			break;
@@ -217,7 +242,6 @@ bool FExrImgMediaReader::ReadFrame(int32 FrameId, const TMap<int32, FImgMediaTil
 				{
 #if defined(PLATFORM_WINDOWS) && PLATFORM_WINDOWS
 					FIntRect TileRegion = CurrentTileSelection.GetVisibleRegion();
-					FIntPoint NumTiles = OutFrame->Info.NumTiles / MipLevelDiv;
 					int32 PixelSize = sizeof(uint16) * OutFrame->Info.NumChannels;
 					TSharedPtr<FSampleConverterParameters> ConverterParams = MakeShared<FSampleConverterParameters>();
 					ConverterParams->FrameInfo = OutFrame->Info;
@@ -226,7 +250,7 @@ bool FExrImgMediaReader::ReadFrame(int32 FrameId, const TMap<int32, FImgMediaTil
 					ConverterParams->NumMipLevels = Loader->GetNumMipLevels();
 					ConverterParams->bCustomExr = OutFrame->Info.FormatName == TEXT("EXR CUSTOM");
 
-					EReadResult ReadResult = ReadTilesCustom((uint16*)MipDataPtr, Image, FrameId, TileRegion, ConverterParams, CurrentMipLevel);
+					EReadResult ReadResult = ReadTilesCustom((uint16*)MipDataPtr, GetMipBufferTotalSize(Dim / MipLevelDiv), Image, FrameId, TileRegion, ConverterParams, CurrentMipLevel);
 					if (ReadResult != Fail)
 					{
 						OutFrame->MipTilesPresent.Emplace(CurrentMipLevel, CurrentTileSelection);
diff --git a/Engine/Plugins/Media/ImgMedia/Source/ImgMedia/Private/Readers/ExrImgMediaReader.h b/Engine/Plugins/Media/ImgMedia/Source/ImgMedia/Private/Readers/ExrImgMediaReader.h
--- a/Engine/Plugins/Media/ImgMedia/Source/ImgMedia/Private/Readers/ExrImgMediaReader.h
+++ b/Engine/Plugins/Media/ImgMedia/Source/ImgMedia/Private/Readers/ExrImgMediaReader.h
@@ -54,13 +54,28 @@ protected:
 	*/
 	struct FSampleConverterParameters
 	{
+		/** General file information including its header. */
 		FImgMediaFrameInfo FrameInfo;
+
+		/** Resolution of the highest quality mip. */
 		FIntPoint FullResolution;
+
+		/** Dimension of the tile including the overscan borders. */
 		FIntPoint TileDimWithBorders;
+
+		/** Used for rendering tiles in bulk regions. */
 		TMap<int32, FIntRect> Viewports;
+
+		/** Number of mip levels read. */
 		int32 NumMipLevels;
+
+		/** Pixel stride in bytes. I.e. 2 bytes per pixel x 3 channels = 6. */
 		int32 PixelSize;
+
+		/** Identifies this exr as custom, therefore all data should be swizzled. */
 		bool bCustomExr;
+
+		/** Indicates if mips stored in individual files.*/
 		bool bMipsInSeparateFiles;
 	};
 
@@ -77,7 +92,7 @@ protected:
 	 * Reads tiles from exr files in tile rows based on Tile region. If frame is pending for cancelation
 	 * stops reading tiles at the current tile row.
 	*/
-	EReadResult ReadTilesCustom(uint16* Buffer, const FString& ImagePath, int32 FrameId, const FIntRect& TileRegion, TSharedPtr<FSampleConverterParameters> ConverterParams, const int32 CurrentMipLevel);
+	EReadResult ReadTilesCustom(uint16* Buffer, int64 BufferSize, const FString& ImagePath, int32 FrameId, const FIntRect& TileRegion, TSharedPtr<FSampleConverterParameters> ConverterParams, const int32 CurrentMipLevel);
 
 	/**
 	 * Sets parameters of our custom format images.
diff --git a/Engine/Plugins/Media/ImgMedia/Source/ImgMedia/Private/Readers/ExrImgMediaReaderGpu.cpp b/Engine/Plugins/Media/ImgMedia/Source/ImgMedia/Private/Readers/ExrImgMediaReaderGpu.cpp
--- a/Engine/Plugins/Media/ImgMedia/Source/ImgMedia/Private/Readers/ExrImgMediaReaderGpu.cpp
+++ b/Engine/Plugins/Media/ImgMedia/Source/ImgMedia/Private/Readers/ExrImgMediaReaderGpu.cpp
@@ -202,7 +202,7 @@ bool FExrImgMediaReaderGpu::ReadFrame(int32 FrameId, const TMap<int32, FImgMedia
 					// read frame data
 					if (bHasTiles || ConverterParams->bCustomExr)
 					{
-						ReadResult = ReadTilesCustom(MipDataPtr, ImagePath, FrameId, TileRegion, ConverterParams, CurrentMipLevel);
+						ReadResult = ReadTilesCustom(MipDataPtr, BufferSize, ImagePath, FrameId, TileRegion, ConverterParams, CurrentMipLevel);
 					}
 					else
 					{
@@ -284,7 +284,9 @@ FExrImgMediaReaderGpu::EReadResult FExrImgMediaReaderGpu::ReadInChunks(uint16* B
 
 	// Since ReadInChunks is only utilized for exr files without tiles and mips, Num Mip levels is always 1.
 	const int32 NumLevels = 1;
-	if (!ChunkReader.OpenExrAndPrepareForPixelReading(ImagePath, Dim.Y, NumLevels))
+	TArray<int32> NumTOffsetsPerLevel;
+	NumTOffsetsPerLevel.Add(Dim.Y);
+	if (!ChunkReader.OpenExrAndPrepareForPixelReading(ImagePath, NumTOffsetsPerLevel, TArray<TArray<int64>>()))
 	{
 		return Fail;
 	}
