ImgMedia: Added a scroll box to process images widget.
#rnx

diff --git a/Engine/Plugins/Media/ImgMedia/Source/ImgMediaEditor/Private/Widgets/SImgMediaProcessImages.cpp b/Engine/Plugins/Media/ImgMedia/Source/ImgMediaEditor/Private/Widgets/SImgMediaProcessImages.cpp
--- a/Engine/Plugins/Media/ImgMedia/Source/ImgMediaEditor/Private/Widgets/SImgMediaProcessImages.cpp
+++ b/Engine/Plugins/Media/ImgMedia/Source/ImgMediaEditor/Private/Widgets/SImgMediaProcessImages.cpp
@@ -55,22 +55,16 @@ void SImgMediaProcessImages::Construct(const FArguments& InArgs)
 	
 	ChildSlot
 	[
-		SNew(SVerticalBox)
-
-		+ SVerticalBox::Slot()
-			.Padding(0, 20, 0, 0)
-			.AutoHeight()
+		SNew(SScrollBox)
 
 		// Add details view.
-		+ SVerticalBox::Slot()
-			.AutoHeight()
+		+ SScrollBox::Slot()
 			[
 				SAssignNew(DetailsViewBox, SBox)
 			]
 			
 		// Add process images button.
-		+ SVerticalBox::Slot()
-			.AutoHeight()
+		+ SScrollBox::Slot()
 			.Padding(4.0f)
 			.HAlign(HAlign_Left)
 			[
