MediaPlate: Allow choices to enable tiling and mip calculations: disabled, plane or sphere mesh. Fixes for (default) plane calculations based on SM_MediaPlateScreen as opposed to the engine mesh.

diff --git a/Engine/Plugins/Media/ImgMedia/Source/ImgMedia/Private/Assets/ImgMediaMipMapInfo.cpp b/Engine/Plugins/Media/ImgMedia/Source/ImgMedia/Private/Assets/ImgMediaMipMapInfo.cpp
--- a/Engine/Plugins/Media/ImgMedia/Source/ImgMedia/Private/Assets/ImgMediaMipMapInfo.cpp
+++ b/Engine/Plugins/Media/ImgMedia/Source/ImgMedia/Private/Assets/ImgMediaMipMapInfo.cpp
@@ -11,7 +11,6 @@
 #include "Engine/StaticMesh.h"
 #include "GameFramework/Actor.h"
 #include "Math/UnrealMathUtility.h"
-#include "MediaTextureTracker.h"
 
 #if WITH_EDITOR
 #include "Editor.h"
@@ -238,11 +237,6 @@ namespace {
 			if (MeshComponent != nullptr)
 			{
 				PlaneSize = 2.0f * MeshComponent->CalcLocalBounds().BoxExtent;
-
-				PlaneVertices[0] = 0.5f * FVector(-PlaneSize.X, -PlaneSize.Y, 0);
-				PlaneVertices[1] = 0.5f * FVector(-PlaneSize.X, PlaneSize.Y, 0);
-				PlaneVertices[2] = 0.5f * FVector(PlaneSize.X, -PlaneSize.Y, 0);
-				PlaneVertices[3] = 0.5f * FVector(PlaneSize.X, PlaneSize.Y, 0);
 			}
 			else
 			{
@@ -268,11 +262,11 @@ namespace {
 			const FTransform MeshTransform = Mesh->GetComponentTransform();
 			const FVector MeshScale = Mesh->GetComponentScale();
 
-			FVector PlaneCornerWS = MeshTransform.TransformPosition(-0.5f * PlaneSize);
-			FVector DirXWS = MeshTransform.TransformVector(FVector(PlaneSize.X, 0, 0));
-			FVector DirYWS = MeshTransform.TransformVector(FVector(0, PlaneSize.Y, 0));
-			FVector TexelOffsetXWS = MeshTransform.TransformVector(FVector(PlaneSize.X / InSequenceInfo.Dim.X, 0, 0));
-			FVector TexelOffsetYWS = MeshTransform.TransformVector(FVector(0, PlaneSize.Y / InSequenceInfo.Dim.Y, 0));
+			FVector PlaneCornerWS = MeshTransform.TransformPosition(FVector(0, -0.5f * PlaneSize.Y, 0.5f * PlaneSize.Z));
+			FVector DirXWS = MeshTransform.TransformVector(FVector(0, PlaneSize.Y, 0));
+			FVector DirYWS = MeshTransform.TransformVector(FVector(0, 0, -PlaneSize.Z));
+			FVector TexelOffsetXWS = MeshTransform.TransformVector(FVector(0, PlaneSize.Y / InSequenceInfo.Dim.X, 0));
+			FVector TexelOffsetYWS = MeshTransform.TransformVector(FVector(0, 0, -PlaneSize.Z / InSequenceInfo.Dim.Y));
 
 			for (const FImgMediaMipMapCameraInfo& CameraInfo : InCameraInfos)
 			{
@@ -315,7 +309,7 @@ namespace {
 					FVector TileCenterWS = PlaneCornerWS + (DirXWS * StepX + DirYWS * StepY);
 
 					// Calculate the tile radius in world space
-					FVector TileSizeWS = (PlaneSize * MeshScale) / FVector(CurrentNumTiles.X, CurrentNumTiles.Y, 1);
+					FVector TileSizeWS = (PlaneSize * MeshScale) / FVector(1, CurrentNumTiles.X, CurrentNumTiles.Y);
 					float TileRadiusInWS = 0.5f * (float)FMath::Sqrt(2 * FMath::Square(TileSizeWS.GetAbsMax()));
 
 					// Now we check if tile spherical bounds are in view.
@@ -436,9 +430,6 @@ namespace {
 		/** Local size of this mesh component. */
 		FVector PlaneSize;
 
-		/** Four plane vertices (in local space). */
-		TStaticArray<FVector, 4> PlaneVertices;
-
 		/** Cached calculating mip levels (at mip0). */
 		mutable TArray<float> CornerMipLevelsCached;
 	};
@@ -541,34 +532,25 @@ FImgMediaMipMapInfo::~FImgMediaMipMapInfo()
 	ClearAllObjects();
 }
 
-void FImgMediaMipMapInfo::AddObject(AActor* InActor, float Width, float LODBias)
+void FImgMediaMipMapInfo::AddObject(AActor* InActor, float Width, float LODBias, EMediaTextureVisibleMipsTiles MeshType)
 {
 	if (InActor != nullptr)
 	{
 		UMeshComponent* MeshComponent = Cast<UMeshComponent>(InActor->FindComponentByClass(UMeshComponent::StaticClass()));
 		if (MeshComponent != nullptr)
 		{
-			FImgMediaMipMapObjectInfo* Info = nullptr;
-
-			if (UStaticMeshComponent* StaticMeshComponent = Cast<UStaticMeshComponent>(MeshComponent))
-			{
-				FString StaticMeshName = StaticMeshComponent->GetStaticMesh()->GetName();
-				if (StaticMeshName.Equals(TEXT("Plane")))
-				{
-					Info = new FPlaneObjectInfo(MeshComponent, LODBias);
-				}
-				else if (StaticMeshName.Equals(TEXT("Sphere")))
-				{
-					Info = new FSphereObjectInfo(MeshComponent, LODBias);
-				}
-			}
-
-			if (Info == nullptr)
+			switch (MeshType)
 			{
-				Info = new FImgMediaMipMapObjectInfo(MeshComponent, LODBias);
+			case EMediaTextureVisibleMipsTiles::Plane:
+				Objects.Add(new FPlaneObjectInfo(MeshComponent, LODBias));
+				break;
+			case EMediaTextureVisibleMipsTiles::Sphere:
+				Objects.Add(new FSphereObjectInfo(MeshComponent, LODBias));
+				break;
+			default:
+				Objects.Add(new FImgMediaMipMapObjectInfo(MeshComponent, LODBias));
+				break;
 			}
-
-			Objects.Add(Info);
 		}
 	}
 }
@@ -605,7 +587,7 @@ void FImgMediaMipMapInfo::AddObjectsUsingThisMediaTexture(UMediaTexture* InMedia
 				AActor* Owner = ObjectInfo->Object.Get();
 				if (Owner != nullptr)
 				{
-					AddObject(Owner, 0.0f, ObjectInfo->MipMapLODBias);
+					AddObject(Owner, 0.0f, ObjectInfo->MipMapLODBias, ObjectInfo->VisibleMipsTilesCalculations);
 				}
 			}
 		}
diff --git a/Engine/Plugins/Media/ImgMedia/Source/ImgMedia/Public/ImgMediaMipMapInfo.h b/Engine/Plugins/Media/ImgMedia/Source/ImgMedia/Public/ImgMediaMipMapInfo.h
--- a/Engine/Plugins/Media/ImgMedia/Source/ImgMedia/Public/ImgMediaMipMapInfo.h
+++ b/Engine/Plugins/Media/ImgMedia/Source/ImgMedia/Public/ImgMediaMipMapInfo.h
@@ -2,12 +2,13 @@
 
 #pragma once
 
+#include "Containers/BitArray.h"
 #include "CoreMinimal.h"
 #include "IMediaOptions.h"
+#include "MediaTextureTracker.h"
 #include "Tickable.h"
 #include "UObject/WeakObjectPtr.h"
 #include "UObject/WeakObjectPtrTemplates.h"
-#include "Containers/BitArray.h"
 
 class AActor;
 class UMeshComponent;
@@ -220,7 +221,7 @@ public:
 	 * @param InActor Object using our img sequence.
 	 * @param Width Width of the object. If < 0, then get the width automatically.
 	 */
-	void AddObject(AActor* InActor, float Width, float LODBias);
+	void AddObject(AActor* InActor, float Width, float LODBias, EMediaTextureVisibleMipsTiles MeshType = EMediaTextureVisibleMipsTiles::None);
 
 	/**
 	 * This object is no longer using our img sequence.
diff --git a/Engine/Plugins/Media/MediaPlate/Source/MediaPlate/Private/MediaPlateComponent.cpp b/Engine/Plugins/Media/MediaPlate/Source/MediaPlate/Private/MediaPlateComponent.cpp
--- a/Engine/Plugins/Media/MediaPlate/Source/MediaPlate/Private/MediaPlateComponent.cpp
+++ b/Engine/Plugins/Media/MediaPlate/Source/MediaPlate/Private/MediaPlateComponent.cpp
@@ -35,6 +35,9 @@ UMediaPlateComponent::UMediaPlateComponent(const FObjectInitializer& ObjectIniti
 
 	// Set up playlist.
 	MediaPlaylist = CreateDefaultSubobject<UMediaPlaylist>(MediaPlaylistName);
+
+	// Default to plane since AMediaPlate defaults to SM_MediaPlateScreen
+	VisibleMipsTilesCalculations = EMediaPlateVisibleMipsTiles::Plane;
 }
 
 void UMediaPlateComponent::OnRegister()
@@ -121,6 +124,7 @@ void UMediaPlateComponent::RegisterWithMediaTextureTracker()
 	MediaTextureTrackerObject = MakeShared<FMediaTextureTrackerObject, ESPMode::ThreadSafe>();
 	MediaTextureTrackerObject->Object = GetOwner();
 	MediaTextureTrackerObject->MipMapLODBias = 0.0f;
+	MediaTextureTrackerObject->VisibleMipsTilesCalculations = static_cast<EMediaTextureVisibleMipsTiles>(VisibleMipsTilesCalculations);
 
 	// Add our texture.
 	TObjectPtr<UMediaTexture> MediaTexture = GetMediaTexture();
@@ -206,6 +210,13 @@ void UMediaPlateComponent::PostEditChangeProperty(FPropertyChangedEvent& Propert
 			}
 		}
 	}
+	else if (PropertyChangedEvent.GetPropertyName() == GET_MEMBER_NAME_CHECKED(ThisClass, VisibleMipsTilesCalculations))
+	{
+		if (MediaTextureTrackerObject != nullptr)
+		{
+			MediaTextureTrackerObject->VisibleMipsTilesCalculations = static_cast<EMediaTextureVisibleMipsTiles>(VisibleMipsTilesCalculations);
+		}
+	}
 }
 
 #endif // WITH_EDITOR
diff --git a/Engine/Plugins/Media/MediaPlate/Source/MediaPlate/Public/MediaPlateComponent.h b/Engine/Plugins/Media/MediaPlate/Source/MediaPlate/Public/MediaPlateComponent.h
--- a/Engine/Plugins/Media/MediaPlate/Source/MediaPlate/Public/MediaPlateComponent.h
+++ b/Engine/Plugins/Media/MediaPlate/Source/MediaPlate/Public/MediaPlateComponent.h
@@ -5,6 +5,7 @@
 #include "Components/ActorComponent.h"
 #include "CoreMinimal.h"
 #include "MediaPlayerOptions.h"
+#include "MediaTextureTracker.h"
 
 #include "MediaPlateComponent.generated.h"
 
@@ -16,6 +17,18 @@ class UMediaSource;
 class UMediaTexture;
 struct FMediaTextureTrackerObject;
 
+/**
+ * List of mesh types used for visible mips and tiles calculations, matching EMediaTextureVisibleMipsTiles.
+ */
+UENUM()
+enum class EMediaPlateVisibleMipsTiles : uint8
+{
+	Disabled = 0,
+	Plane = 1,
+	Sphere
+};
+
+
 /**
  * This is a component for AMediaPlate that can play and show media in the world.
  */
@@ -96,6 +109,13 @@ public:
 	UPROPERTY(EditAnywhere, AdvancedDisplay, Category = "MediaPlate")
 	float SmartCacheTimeToLookAhead = 0.2f;
 
+	/** Specify type of mesh used for visible mips and tiles calculations.
+	 * Use /MediaPlate/SM_MediaPlateScreen for the plane.
+	 * Use /Engine/BasicShapes/Sphere for the sphere.
+	 */
+	UPROPERTY(EditAnywhere, AdvancedDisplay, Category = "MediaPlate")
+	EMediaPlateVisibleMipsTiles VisibleMipsTilesCalculations;
+
 	/**
 	 * Adds our media texture to the media texture tracker.
 	 */
diff --git a/Engine/Source/Runtime/MediaAssets/Public/MediaTextureTracker.h b/Engine/Source/Runtime/MediaAssets/Public/MediaTextureTracker.h
--- a/Engine/Source/Runtime/MediaAssets/Public/MediaTextureTracker.h
+++ b/Engine/Source/Runtime/MediaAssets/Public/MediaTextureTracker.h
@@ -10,6 +10,13 @@
 class AActor;
 class UMediaTexture;
 
+enum class EMediaTextureVisibleMipsTiles : uint8
+{
+	None = 0,
+	Plane = 1,
+	Sphere
+};
+
 /** Holds info on a single object. */
 struct FMediaTextureTrackerObject
 {
@@ -17,6 +24,8 @@ struct FMediaTextureTrackerObject
 	TWeakObjectPtr<class AActor> Object;
 	/** LOD bias for the mipmap level. */
 	float MipMapLODBias;
+	/** Specify type of mesh used for visible mips and tiles calculations. */
+	EMediaTextureVisibleMipsTiles VisibleMipsTilesCalculations;
 };
 
 /**
