ImgMediaReader ExrReaderGPU: Part2 - Added shader conversion support for partial tiles. Also added an additional buffer that contains tile description, which reduces calculations on GPU improving performance.

diff --git a/Engine/Plugins/Media/ImgMedia/Shaders/Private/ExrSwizzler.usf b/Engine/Plugins/Media/ImgMedia/Shaders/Private/ExrSwizzler.usf
--- a/Engine/Plugins/Media/ImgMedia/Shaders/Private/ExrSwizzler.usf
+++ b/Engine/Plugins/Media/ImgMedia/Shaders/Private/ExrSwizzler.usf
@@ -2,9 +2,23 @@
 
 #include "/Engine/Private/Common.ush"
 
+#if PARTIAL_TILES
+struct FTileDesc
+{
+	int2 TileResolution;
+	uint BufferOffset;
+	int AlignPadding;
+};
+#endif
+
 /** The unswizzled buffer packed into 4 byte buffer stride. */
 StructuredBuffer<uint> UnswizzledBuffer;
 
+#if PARTIAL_TILES
+/** The unswizzled buffer packed into 4 byte buffer stride. */
+StructuredBuffer<FTileDesc> TileDescBuffer;
+#endif
+
 /** Texture size that the buffer is supposed to be unpacked to. */
 int2 TextureSize;
 
@@ -64,6 +78,32 @@ void MainVS(
 	DrawRectangle(InPosition, InTexCoord, OutPosition, OutUVAndScreenPos);
 }
 
+#if PARTIAL_TILES
+void GetStartScanlinePosition_PartialTiles(in int2 PixelPos, out int2 SampleSize, out int StartPosition)
+{
+	int2 TileCoord = int2(ceil(PixelPos.x / TileSize.x), ceil(PixelPos.y / TileSize.y));
+	FTileDesc TileDesc = TileDescBuffer[TileCoord.y * NumTiles.x + TileCoord.x];
+
+	SampleSize = TileDesc.TileResolution;
+
+	int xCoord = (PixelPos.x) % TileSize.x;
+	int yCoord = (PixelPos.y) % TileSize.y;
+
+	// Padding in buffer elements (uint)
+	const int TilePadding = 10;
+
+	// Offset is in bytes while our buffer is a packed uint16 into uint32. Therefore divide by 2. 
+	int PreviousTilesOffset = TileDesc.BufferOffset / 2;
+
+	// Current tile position.
+#if CUSTOM_EXR
+	StartPosition = PreviousTilesOffset + yCoord * TileDesc.TileResolution.x * NumChannels + xCoord * NumChannels;
+#else
+	StartPosition = PreviousTilesOffset + yCoord * TileDesc.TileResolution.x * NumChannels + xCoord + TilePadding;
+#endif //CUSTOM_EXR
+}
+#endif
+
 void GetStartScanlinePosition(in int2 PixelPos, out int2 SampleSize, out int StartPosition)
 {
 	// Depending on the EXR file it could contain multi-channel data.
@@ -105,14 +145,20 @@ void GetStartScanlinePosition(in int2 PixelPos, out int2 SampleSize, out int Sta
 /** Pixel Shader. */
 half4 MainPS(noperspective float4 UVAndScreenPos : TEXCOORD0) : SV_Target0
 {
-	float2 UV = UVAndScreenPos.xy;
 
+	float2 UV = UVAndScreenPos.xy;
+	
 	int2 SampleSize;
 
 	int2 PixelPos = UV * TextureSize;
+
 	int StartBufferPosition;
 
+#if PARTIAL_TILES
+	GetStartScanlinePosition_PartialTiles(PixelPos, SampleSize, StartBufferPosition);
+#else
 	GetStartScanlinePosition(PixelPos, SampleSize, StartBufferPosition);
+#endif
 
 	int ChannelIndex = 0;
 
diff --git a/Engine/Plugins/Media/ImgMedia/Source/ExrReaderGpu/Private/ExrReaderGpu.cpp b/Engine/Plugins/Media/ImgMedia/Source/ExrReaderGpu/Private/ExrReaderGpu.cpp
--- a/Engine/Plugins/Media/ImgMedia/Source/ExrReaderGpu/Private/ExrReaderGpu.cpp
+++ b/Engine/Plugins/Media/ImgMedia/Source/ExrReaderGpu/Private/ExrReaderGpu.cpp
@@ -16,6 +16,12 @@ THIRD_PARTY_INCLUDES_START
 THIRD_PARTY_INCLUDES_END
 PRAGMA_DEFAULT_VISIBILITY_END
 
+static TAutoConsoleVariable<bool> CVarForceTileDescBufferExrGpuReader(
+	TEXT("r.ExrReaderGPU.ForceTileDescBuffer"),
+	true,
+	TEXT("Calculates tile description and offsets on CPU and provides a Structured buffer.\n")
+	TEXT("to be used to access tile description on GPU\n"));
+
 namespace
 {
 	bool ReadBytes(FILE* FileHandle, char* Buffer, int32 Length)
@@ -255,6 +261,7 @@ bool FExrReader::GenerateTextureData(uint16* Buffer, int32 BufferSize, FString F
 void FExrReader::CalculateTileOffsets
 	( TArray<int32>& OutNumTilesPerLevel
 	, TArray<TArray<int64>>& OutCustomOffsets
+	, TArray<TArray<FTileDesc>>& OutPartialTileInfo
 	, const FIntPoint& FullTextureResolution
 	, const FIntPoint& FullResolutionInTiles
 	, const FIntPoint& TileDimWithBorders
@@ -266,7 +273,6 @@ void FExrReader::CalculateTileOffsets
 
 	OutCustomOffsets.SetNum(NumMipLevels);
 	// Custom exr files don't have tile offsets, therefore we need to calculate these manually.
-	if (bCustomExr)
 	{
 		int64 CurrentPosition = 0;
 		for (int32 MipLevel = 0; MipLevel < NumMipLevels; MipLevel++)
@@ -287,7 +293,24 @@ void FExrReader::CalculateTileOffsets
 			int32 NumActualTiles = DimensionInTiles_PartialTiles.X * DimensionInTiles_PartialTiles.Y;
 
 			// Custom exr has one Vanilla tile per mip level that contains custom tiles.
-			OutNumTilesPerLevel.Add(1);
+			if (bCustomExr)
+			{
+				OutNumTilesPerLevel.Add(1);
+			}
+			else
+			{
+				OutNumTilesPerLevel.Add(DimensionInTiles_PartialTiles.X * DimensionInTiles_PartialTiles.Y);
+			}
+
+			bool bHasPartialTiles = (DimensionInTiles_PartialTiles != DimensionInTiles_CompleteTiles) || CVarForceTileDescBufferExrGpuReader.GetValueOnAnyThread();
+
+			OutPartialTileInfo.Add({});
+			TArray<FTileDesc>& TileInfoList = OutPartialTileInfo[MipLevel];
+
+			if (!bCustomExr && !bHasPartialTiles)
+			{
+				continue;
+			}
 
 			// Resolution of the texture in pixels for this mip level.
 			FIntPoint MipResolution = FullTextureResolution / MipDiv;
@@ -295,58 +318,62 @@ void FExrReader::CalculateTileOffsets
 			// Resolution of the partial tile in the bottom right corner.
 			FIntPoint PartialTileResolution = FIntPoint(MipResolution.X % TileDimWithBorders.X, MipResolution.Y % TileDimWithBorders.Y);
 
-			bool bHasPartialTiles = DimensionInTiles_PartialTiles != DimensionInTiles_CompleteTiles;
-
 			TArray<int64>& CurrentMipOffsets = OutCustomOffsets[MipLevel];
 			CurrentMipOffsets.SetNum(NumActualTiles);
-			for (int TileIndex = 0; TileIndex < NumActualTiles; TileIndex++)
+
+			if (bHasPartialTiles)
 			{
-				if (TileIndex == 0)
-				{
-					CurrentPosition += FExrReader::TILE_PADDING;
-				}
+				TileInfoList.SetNum(NumActualTiles);
+			}
 
-				if (MipLevel == 0 && TileIndex == 0)
-				{
-					CurrentMipOffsets[TileIndex] = CurrentPosition;
-					continue;
-				}
+			if (bCustomExr)
+			{
+				CurrentPosition += FExrReader::TILE_PADDING;
+			}
+
+			int64 MipOffsetStart = CurrentPosition;
 
+			for (int TileIndex = 0; TileIndex < NumActualTiles; TileIndex++)
+			{
 				FIntPoint TileDim(TileDimWithBorders);
 
-				if (bHasPartialTiles)
+				if (DimensionInTiles_PartialTiles.X != DimensionInTiles_CompleteTiles.X)
 				{
 					// If true - this is a partial tile in X dimension.
-					if (((TileIndex) % DimensionInTiles_PartialTiles.X) == 0)
+					if (TileIndex != 0 &&
+						((TileIndex + 1) % DimensionInTiles_PartialTiles.X) == 0)
 					{
 						TileDim.X = PartialTileResolution.X;
 					}
+				}
 
+				if (DimensionInTiles_PartialTiles.Y != DimensionInTiles_CompleteTiles.Y)
+				{
 					// If true - this is a partial tile in Y dimension.
-					if (TileIndex > (DimensionInTiles_PartialTiles.X * DimensionInTiles_CompleteTiles.Y))
+					if (TileIndex >= (DimensionInTiles_PartialTiles.X * DimensionInTiles_CompleteTiles.Y))
 					{
 						TileDim.Y = PartialTileResolution.Y;
 					}
 				}
 
+				if (bHasPartialTiles)
+				{
+					TileInfoList[TileIndex] = { TileDim, (uint32)(CurrentPosition - MipOffsetStart) };
+				}
+
+				CurrentMipOffsets[TileIndex] = CurrentPosition;
+
 				// Tile offset.
 				CurrentPosition += TileDim.X * TileDim.Y * PixelSize;
 
-				// First tile at each mip level contains padding.
-				CurrentMipOffsets[TileIndex] = CurrentPosition;
+				// Vanilla exr has 20 byte padding at the beginning of each tile.
+				if (!bCustomExr)
+				{
+					CurrentPosition += FExrReader::TILE_PADDING;
+				}
 			}
 		}
 	}
-	else
-	{
-		for (int32 MipLevel = 0; MipLevel < NumMipLevels; MipLevel++)
-		{
-			int32 MipDiv = 1 << MipLevel;
-			int32 DimX = FMath::Max(1, FMath::CeilToInt(float(FullResolutionInTiles.X) / MipDiv));
-			int32 DimY = FMath::Max(1, FMath::CeilToInt(float(FullResolutionInTiles.Y) / MipDiv));
-			OutNumTilesPerLevel.Add(DimX * DimY);
-		}
-	}
 }
 
 bool FExrReader::OpenExrAndPrepareForPixelReading(FString FilePath, const TArray<int32>& NumOffsetsPerLevel, TArray<TArray<int64>>&& CustomOffsets, bool bInCustomExr)
diff --git a/Engine/Plugins/Media/ImgMedia/Source/ExrReaderGpu/Public/ExrReaderGpu.h b/Engine/Plugins/Media/ImgMedia/Source/ExrReaderGpu/Public/ExrReaderGpu.h
--- a/Engine/Plugins/Media/ImgMedia/Source/ExrReaderGpu/Public/ExrReaderGpu.h
+++ b/Engine/Plugins/Media/ImgMedia/Source/ExrReaderGpu/Public/ExrReaderGpu.h
@@ -17,6 +17,15 @@ class EXRREADERGPU_API FExrReader
 {
 public:
 
+	struct FTileDesc
+	{
+		FIntPoint TileResolution;
+		uint32 BufferOffset;
+
+		// Optimization: align to 128 bit. 
+		int32 AlignPadding;
+	};
+
 	/** 
 	*	This is based on Open Exr implementation.
 	*	We need to handle all of these in the future. At the moment only INCREASING_Y is 
@@ -69,8 +78,10 @@ public:
 	 * To be able to read tiles, especially in cases where tiles vary in sizes, we need to calculate offsets manually
 	 * for custom exr files, since we don't write out offsets.
 	*/
-	static void CalculateTileOffsets(TArray<int32>& OutNumTilesPerLevel
+	static void CalculateTileOffsets
+		( TArray<int32>& OutNumTilesPerLevel
 		, TArray<TArray<int64>>& OutCustomOffsets
+		, TArray<TArray<FTileDesc>>& OutPartialTileInfo
 		, const FIntPoint& FullTextureResolution
 		, const FIntPoint& FullResolutionInTiles
 		, const FIntPoint& TileDimWithBorders
diff --git a/Engine/Plugins/Media/ImgMedia/Source/ExrReaderGpu/Public/ExrSwizzlingShader.h b/Engine/Plugins/Media/ImgMedia/Source/ExrReaderGpu/Public/ExrSwizzlingShader.h
--- a/Engine/Plugins/Media/ImgMedia/Source/ExrReaderGpu/Public/ExrSwizzlingShader.h
+++ b/Engine/Plugins/Media/ImgMedia/Source/ExrReaderGpu/Public/ExrSwizzlingShader.h
@@ -43,7 +43,8 @@ class EXRREADERGPU_API FExrSwizzlePS : public FGlobalShader
 	class FRgbaSwizzle : SHADER_PERMUTATION_INT("NUM_CHANNELS", 4);
 	class FRenderTiles : SHADER_PERMUTATION_BOOL("RENDER_TILES");
 	class FCustomExr : SHADER_PERMUTATION_BOOL("CUSTOM_EXR");
-	using FPermutationDomain = TShaderPermutationDomain<FRgbaSwizzle, FRenderTiles, FCustomExr>;
+	class FPartialTiles : SHADER_PERMUTATION_BOOL("PARTIAL_TILES");
+	using FPermutationDomain = TShaderPermutationDomain<FRgbaSwizzle, FRenderTiles, FCustomExr, FPartialTiles>;
 
 	static bool ShouldCompilePermutation(const FGlobalShaderPermutationParameters& Parameters) 
 	{
@@ -52,6 +53,7 @@ class EXRREADERGPU_API FExrSwizzlePS : public FGlobalShader
 
 	BEGIN_SHADER_PARAMETER_STRUCT(FParameters, )
 		SHADER_PARAMETER_SRV(StructuredBuffer<uint>, UnswizzledBuffer)
+		SHADER_PARAMETER_SRV(StructuredBuffer<FTileDesc>, TileDescBuffer)
 		SHADER_PARAMETER(FIntPoint, TextureSize)
 		SHADER_PARAMETER(FIntPoint, TileSize)
 		SHADER_PARAMETER(FIntPoint, NumTiles)
diff --git a/Engine/Plugins/Media/ImgMedia/Source/ImgMedia/Private/Readers/ExrImgMediaReader.cpp b/Engine/Plugins/Media/ImgMedia/Source/ImgMedia/Private/Readers/ExrImgMediaReader.cpp
--- a/Engine/Plugins/Media/ImgMedia/Source/ImgMedia/Private/Readers/ExrImgMediaReader.cpp
+++ b/Engine/Plugins/Media/ImgMedia/Source/ImgMedia/Private/Readers/ExrImgMediaReader.cpp
@@ -56,7 +56,7 @@ bool FExrImgMediaReader::GetFrameInfo(const FString& ImagePath, FImgMediaFrameIn
 }
 
 FExrImgMediaReader::EReadResult FExrImgMediaReader::ReadTilesCustom
-	( uint16* Buffer
+(uint16* Buffer
 	, int64 BufferSize
 	, const FString& ImagePath
 	, int32 FrameId
@@ -72,71 +72,79 @@ FExrImgMediaReader::EReadResult FExrImgMediaReader::ReadTilesCustom
 	int MipLevelDiv = 1 << CurrentMipLevel;
 
 	FIntPoint DimensionInTiles
-		( FMath::CeilToInt(float(ConverterParams->FrameInfo.NumTiles.X) / MipLevelDiv)
+	(FMath::CeilToInt(float(ConverterParams->FrameInfo.NumTiles.X) / MipLevelDiv)
 		, FMath::CeilToInt(float(ConverterParams->FrameInfo.NumTiles.Y) / MipLevelDiv));
 
 	TArray<int32> NumTilesPerLevel;
 	TArray<TArray<int64>> CustomOffsets;
 	int32 NumMipLevels = ConverterParams->bMipsInSeparateFiles ? 1 : ConverterParams->NumMipLevels;
+	bool bHasPartialTiles = false;
 
 	FExrReader::CalculateTileOffsets(
-		NumTilesPerLevel, 
-		CustomOffsets, 
-		ConverterParams->FullResolution, 
+		NumTilesPerLevel,
+		CustomOffsets,
+		ConverterParams->TileInfoPerMipLevel,
+		ConverterParams->FullResolution,
 		ConverterParams->FrameInfo.NumTiles,
 		ConverterParams->TileDimWithBorders,
-		NumMipLevels, 
+		NumMipLevels,
 		ConverterParams->PixelSize,
 		ConverterParams->bCustomExr);
 
-	if (!ChunkReader.OpenExrAndPrepareForPixelReading(ImagePath, NumTilesPerLevel, MoveTemp(CustomOffsets), ConverterParams->bCustomExr))
 	{
-		return Fail;
+		TRACE_CPUPROFILER_EVENT_SCOPE_TEXT(*FString::Printf(TEXT("FExrImgMediaReader_ReadTilesCustom_OpenFile")));
+		if (!ChunkReader.OpenExrAndPrepareForPixelReading(ImagePath, NumTilesPerLevel, MoveTemp(CustomOffsets), ConverterParams->bCustomExr))
+		{
+			return Fail;
+		}
 	}
-
-	for (int32 TileRow = TileRegion.Min.Y; TileRow < TileRegion.Max.Y; TileRow++)
 	{
+		TRACE_CPUPROFILER_EVENT_SCOPE_TEXT(*FString::Printf(TEXT("FExrImgMediaReader_ReadTilesCustom_ReadTiles")));
 
-		// Check to see if the frame was canceled.
+		for (int32 TileRow = TileRegion.Min.Y; TileRow < TileRegion.Max.Y; TileRow++)
 		{
-			FScopeLock RegionScopeLock(&CanceledFramesCriticalSection);
-			if (CanceledFrames.Remove(FrameId) > 0)
+			// Check to see if the frame was canceled.
 			{
-				UE_LOG(LogImgMedia, Warning, TEXT("Reader %p: Canceling Frame %i At tile row # %i"), this, FrameId, TileRow);
-				bResult = Cancelled;
-				break;
+				FScopeLock RegionScopeLock(&CanceledFramesCriticalSection);
+				if (CanceledFrames.Remove(FrameId) > 0)
+				{
+					UE_LOG(LogImgMedia, Warning, TEXT("Reader %p: Canceling Frame %i At tile row # %i"), this, FrameId, TileRow);
+					bResult = Cancelled;
+					break;
+				}
 			}
-		}
 
-		const uint16 Padding = ConverterParams->bCustomExr ? 0 : FExrReader::TILE_PADDING;
-		const int64 TileByteStride = ConverterParams->PixelSize * ConverterParams->TileDimWithBorders.X * ConverterParams->TileDimWithBorders.Y + Padding;
-		const int StartTileIndex = TileRow * DimensionInTiles.X + TileRegion.Min.X;
+			const uint16 Padding = ConverterParams->bCustomExr ? 0 : FExrReader::TILE_PADDING;
+			const int64 TileByteStride = ConverterParams->PixelSize * ConverterParams->TileDimWithBorders.X * ConverterParams->TileDimWithBorders.Y + Padding;
+			const int StartTileIndex = TileRow * DimensionInTiles.X + TileRegion.Min.X;
 
-		bool bLastTile = TileRow + 1 == DimensionInTiles.Y && TileRegion.Max.X == DimensionInTiles.X;
-		const int EndTileIndex = TileRow * DimensionInTiles.X + TileRegion.Max.X;
+			bool bLastTile = TileRow + 1 == DimensionInTiles.Y && TileRegion.Max.X == DimensionInTiles.X;
+			const int EndTileIndex = TileRow * DimensionInTiles.X + TileRegion.Max.X;
 
-		int MipLevel = ConverterParams->bMipsInSeparateFiles ? 0 : CurrentMipLevel;
-		ChunkReader.SeekTileWithinFile(StartTileIndex, MipLevel, CurrentBufferPos);
-		int64 ByteOffsetStart = 0;
-		int64 ByteOffsetEnd = 0;
+			int MipLevel = ConverterParams->bMipsInSeparateFiles ? 0 : CurrentMipLevel;
+			ChunkReader.SeekTileWithinFile(StartTileIndex, MipLevel, CurrentBufferPos);
+			int64 ByteOffsetStart = 0;
+			int64 ByteOffsetEnd = 0;
 
-		if (!ChunkReader.GetByteOffsetForTile(StartTileIndex, MipLevel, ByteOffsetStart)
-			|| !ChunkReader.GetByteOffsetForTile(EndTileIndex, MipLevel, ByteOffsetEnd))
-		{
-			bResult = Fail;
-			break;
-		}
+			if (!ChunkReader.GetByteOffsetForTile(StartTileIndex, MipLevel, ByteOffsetStart)
+				|| !ChunkReader.GetByteOffsetForTile(EndTileIndex, MipLevel, ByteOffsetEnd))
+			{
+				bResult = Fail;
+				break;
+			}
 
-		int64 ByteChunkToRead = bLastTile
-			? FMath::Min(ByteOffsetEnd - ByteOffsetStart, BufferSize - CurrentBufferPos)
-			: (ByteOffsetEnd - ByteOffsetStart);
-		if (!ChunkReader.ReadExrImageChunk(reinterpret_cast<char*>(Buffer) + CurrentBufferPos, ByteChunkToRead))
-		{
-			bResult = Fail;
-			break;
+			// If this is the last tile, make sure chunk reader reads till the end of the buffer we write into.
+			int64 ByteChunkToRead = bLastTile
+				? FMath::Min(ByteOffsetEnd - ByteOffsetStart, BufferSize - CurrentBufferPos)
+				: (ByteOffsetEnd - ByteOffsetStart);
+
+			if (!ChunkReader.ReadExrImageChunk(reinterpret_cast<char*>(Buffer) + CurrentBufferPos, ByteChunkToRead))
+			{
+				bResult = Fail;
+				break;
+			}
 		}
 	}
-
 	if (!ChunkReader.CloseExrFile())
 	{
 		return Fail;
diff --git a/Engine/Plugins/Media/ImgMedia/Source/ImgMedia/Private/Readers/ExrImgMediaReader.h b/Engine/Plugins/Media/ImgMedia/Source/ImgMedia/Private/Readers/ExrImgMediaReader.h
--- a/Engine/Plugins/Media/ImgMedia/Source/ImgMedia/Private/Readers/ExrImgMediaReader.h
+++ b/Engine/Plugins/Media/ImgMedia/Source/ImgMedia/Private/Readers/ExrImgMediaReader.h
@@ -7,12 +7,12 @@
 #if IMGMEDIA_EXR_SUPPORTED_PLATFORM
 
 #include "IImgMediaReader.h"
+#include "ExrReaderGpu.h"
 
 class FImgMediaLoader;
 class FOpenExrHeaderReader;
 class FRgbaInputFile;
 
-
 /**
  * Implements a reader for EXR image sequences.
  */
@@ -66,6 +66,11 @@ protected:
 		/** Used for rendering tiles in bulk regions. */
 		TMap<int32, FIntRect> Viewports;
 
+#if defined(PLATFORM_WINDOWS) && PLATFORM_WINDOWS
+		/** Contain information about individual tiles. Used to convert buffer data into a 2D Texture. */
+		TArray<TArray<FExrReader::FTileDesc>> TileInfoPerMipLevel;
+#endif
+
 		/** Number of mip levels read. */
 		int32 NumMipLevels;
 
diff --git a/Engine/Plugins/Media/ImgMedia/Source/ImgMedia/Private/Readers/ExrImgMediaReaderGpu.cpp b/Engine/Plugins/Media/ImgMedia/Source/ImgMedia/Private/Readers/ExrImgMediaReaderGpu.cpp
--- a/Engine/Plugins/Media/ImgMedia/Source/ImgMedia/Private/Readers/ExrImgMediaReaderGpu.cpp
+++ b/Engine/Plugins/Media/ImgMedia/Source/ImgMedia/Private/Readers/ExrImgMediaReaderGpu.cpp
@@ -30,6 +30,7 @@
 
 
 DECLARE_GPU_STAT_NAMED(ExrImgMediaReaderGpu, TEXT("ExrImgMediaReaderGpu"));
+DECLARE_GPU_STAT_NAMED(ExrImgMediaReaderGpu_MipRender, TEXT("ExrImgMediaReaderGpu_MipRender"));
 
 namespace {
 
@@ -164,7 +165,11 @@ bool FExrImgMediaReaderGpu::ReadFrame(int32 FrameId, const TMap<int32, FImgMedia
 
 		for (const TPair<int32, FImgMediaTileSelection>& TilesPerMip : InMipTiles)
 		{
+
 			const int32 CurrentMipLevel = TilesPerMip.Key;
+
+			TRACE_CPUPROFILER_EVENT_SCOPE_TEXT(*FString::Printf(TEXT("FExrImgMediaReaderGpu_ReadMip %d"), CurrentMipLevel));
+
 			const FImgMediaTileSelection& CurrentTileSelection = TilesPerMip.Value;
 
 			bool ReadThisMip = true;
@@ -368,9 +373,11 @@ TSharedPtr<IMediaTextureSampleConverter, ESPMode::ThreadSafe> FExrImgMediaReader
 		FIntPoint Dim = ConverterParams->FullResolution;
 		for (int32 MipLevel = 0; MipLevel < ConverterParams->NumMipLevels; ++MipLevel)
 		{
+			SCOPED_GPU_STAT(RHICmdList, ExrImgMediaReaderGpu_MipRender);
 			int MipLevelDiv = 1 << MipLevel;
 			Dim = ConverterParams->FullResolution / MipLevelDiv;
 
+
 			FStructuredBufferPoolItemSharedPtr BufferData = BufferDataArray[MipLevel];
 			if (BufferData.IsValid())
 			{
@@ -391,13 +398,33 @@ TSharedPtr<IMediaTextureSampleConverter, ESPMode::ThreadSafe> FExrImgMediaReader
 				PermutationVector.Set<FExrSwizzlePS::FRgbaSwizzle>(ConverterParams->FrameInfo.NumChannels - 1);
 				PermutationVector.Set<FExrSwizzlePS::FRenderTiles>(ConverterParams->FrameInfo.bHasTiles || ConverterParams->bCustomExr);
 				PermutationVector.Set<FExrSwizzlePS::FCustomExr>(ConverterParams->bCustomExr);
+				PermutationVector.Set<FExrSwizzlePS::FPartialTiles>(false);
+
 				FExrSwizzlePS::FParameters Parameters = FExrSwizzlePS::FParameters();
 				Parameters.TextureSize = Dim;
 				Parameters.TileSize = ConverterParams->TileDimWithBorders;
 				Parameters.NumChannels = ConverterParams->FrameInfo.NumChannels;
 				if (ConverterParams->FrameInfo.bHasTiles)
 				{
-					Parameters.NumTiles = Dim / ConverterParams->TileDimWithBorders;
+					Parameters.NumTiles = FIntPoint(FMath::CeilToInt(float(Dim.X) / ConverterParams->TileDimWithBorders.X), FMath::CeilToInt(float(Dim.Y) / ConverterParams->TileDimWithBorders.Y));
+				}
+
+				if (ConverterParams->FrameInfo.bHasTiles && 
+					(ConverterParams->TileInfoPerMipLevel.Num() > MipLevel && ConverterParams->TileInfoPerMipLevel[MipLevel].Num() > 0))
+				{
+					TRACE_CPUPROFILER_EVENT_SCOPE_TEXT(*FString::Printf(TEXT("FExrImgMediaReaderGpu_TileDesc")));
+
+					FBufferRHIRef BufferRef;
+					FRHIResourceCreateInfo CreateInfo(TEXT("FExrImgMediaReaderGpu_TileDesc"));
+					const uint32 BytesPerElement = sizeof(FExrReader::FTileDesc);
+					const uint32 NumElements = ConverterParams->TileInfoPerMipLevel[MipLevel].Num();
+
+					BufferRef = RHICreateStructuredBuffer(BytesPerElement, BytesPerElement*NumElements, BUF_ShaderResource | BUF_Dynamic | BUF_FastVRAM, CreateInfo);
+					void* MappedBuffer = RHILockBuffer(BufferRef, 0, NumElements, RLM_WriteOnly);
+					FMemory::Memcpy(MappedBuffer, &ConverterParams->TileInfoPerMipLevel[MipLevel][0], sizeof(FExrReader::FTileDesc) * ConverterParams->TileInfoPerMipLevel[MipLevel].Num());
+					RHIUnlockBuffer(BufferRef);
+					Parameters.TileDescBuffer = RHICreateShaderResourceView(BufferRef);
+					PermutationVector.Set<FExrSwizzlePS::FPartialTiles>(true);
 				}
 
 				Parameters.UnswizzledBuffer = RHICreateShaderResourceView(BufferData->BufferRef);
@@ -435,6 +462,8 @@ TSharedPtr<IMediaTextureSampleConverter, ESPMode::ThreadSafe> FExrImgMediaReader
 
 FStructuredBufferPoolItemSharedPtr FExrImgMediaReaderGpu::AllocateGpuBufferFromPool(uint32 AllocSize, bool bWait)
 {
+	TRACE_CPUPROFILER_EVENT_SCOPE_TEXT(*FString::Printf(TEXT("FExrImgMediaReaderGpu_AllocBuffer %d"), AllocSize));
+
 	// This function is attached to the shared pointer and is used to return any allocated memory to staging pool.
 	auto BufferDeleter = [AllocSize](FStructuredBufferPoolItem* ObjectToDelete) {
 		ReturnGpuBufferToStagingPool(AllocSize, ObjectToDelete);
