ImgMedia: Scene view extension improvements: store hidden or show only primitives to avoid needless mipmapinfo calculations, adjusted priority, moved to public because it is exposed through the module interface.

diff --git a/Engine/Plugins/Media/ImgMedia/Source/ImgMedia/Private/Assets/ImgMediaMipMapInfo.cpp b/Engine/Plugins/Media/ImgMedia/Source/ImgMedia/Private/Assets/ImgMediaMipMapInfo.cpp
--- a/Engine/Plugins/Media/ImgMedia/Source/ImgMedia/Private/Assets/ImgMediaMipMapInfo.cpp
+++ b/Engine/Plugins/Media/ImgMedia/Source/ImgMedia/Private/Assets/ImgMediaMipMapInfo.cpp
@@ -4,6 +4,7 @@
 
 #include "IImgMediaModule.h"
 #include "ImgMediaPrivate.h"
+#include "ImgMediaSceneViewExtension.h"
 
 #include "Async/Async.h"
 #include "Components/StaticMeshComponent.h"
@@ -166,6 +167,16 @@ int32 FImgMediaTileSelection::NumVisibleTiles() const
 	return NumVisibleTiles;
 }
 
+namespace {
+	bool IsPrimitiveComponentHidden(FPrimitiveComponentId ComponentId, const FImgMediaViewInfo& ViewInfo)
+	{
+		bool bIsPrimitiveContained = ViewInfo.PrimitiveComponentIds.Contains(ComponentId);
+
+		// The primitive component id is either part of the hidden list, or not in the show-only list.
+		return ViewInfo.bPrimitiveHiddenMode ? bIsPrimitiveContained : !bIsPrimitiveContained;
+	}
+}
+
 
 FImgMediaMipMapObjectInfo::FImgMediaMipMapObjectInfo(UMeshComponent* InMeshComponent, float InLODBias)
 	: MeshComponent(InMeshComponent)
@@ -179,14 +190,13 @@ UMeshComponent* FImgMediaMipMapObjectInfo::GetMeshComponent() const
 	return MeshComponent.Get();
 }
 
-bool FImgMediaMipMapObjectInfo::CalculateVisibleTiles(const TArray<FImgMediaMipMapCameraInfo>& InCameraInfos, const FSequenceInfo& InSequenceInfo, TMap<int32, FImgMediaTileSelection>& VisibleTiles) const
+void FImgMediaMipMapObjectInfo::CalculateVisibleTiles(const TArray<FImgMediaViewInfo>& InViewInfos, const FSequenceInfo& InSequenceInfo, TMap<int32, FImgMediaTileSelection>& VisibleTiles) const
 {
 	// We simply add fully visible regions for all mip levels
 	for (int32 MipLevel = 0; MipLevel < InSequenceInfo.NumMipLevels; ++MipLevel)
 	{
 		VisibleTiles.Add(MipLevel, FImgMediaTileSelection::CreateForTargetMipLevel(InSequenceInfo.NumTiles.X, InSequenceInfo.NumTiles.Y, MipLevel, true));
 	}
-	return true;
 }
 
 namespace {
@@ -207,14 +217,14 @@ namespace {
 	}
 
 	// Approximates hardware mip level selection.
-	bool CalculateMipLevel(const FImgMediaMipMapCameraInfo& CameraInfo, const FVector& TexelWS, const FVector& OffsetXWS, const FVector& OffsetYWS, float& OutMipLevel)
+	bool CalculateMipLevel(const FImgMediaViewInfo& ViewInfo, const FVector& TexelWS, const FVector& OffsetXWS, const FVector& OffsetYWS, float& OutMipLevel)
 	{
 		FVector2D TexelScreenSpace[3];
 
 		bool bValid = true;
-		bValid &= ProjectWorldToScreenFast(TexelWS, CameraInfo.ViewportRect, CameraInfo.ViewProjectionMatrix, TexelScreenSpace[0]);
-		bValid &= ProjectWorldToScreenFast(TexelWS + OffsetXWS, CameraInfo.ViewportRect, CameraInfo.ViewProjectionMatrix, TexelScreenSpace[1]);
-		bValid &= ProjectWorldToScreenFast(TexelWS + OffsetYWS, CameraInfo.ViewportRect, CameraInfo.ViewProjectionMatrix, TexelScreenSpace[2]);
+		bValid &= ProjectWorldToScreenFast(TexelWS, ViewInfo.ViewportRect, ViewInfo.ViewProjectionMatrix, TexelScreenSpace[0]);
+		bValid &= ProjectWorldToScreenFast(TexelWS + OffsetXWS, ViewInfo.ViewportRect, ViewInfo.ViewProjectionMatrix, TexelScreenSpace[1]);
+		bValid &= ProjectWorldToScreenFast(TexelWS + OffsetYWS, ViewInfo.ViewportRect, ViewInfo.ViewProjectionMatrix, TexelScreenSpace[2]);
 
 		if (bValid)
 		{
@@ -244,12 +254,12 @@ namespace {
 			}
 		}
 
-		bool CalculateVisibleTiles(const TArray<FImgMediaMipMapCameraInfo>& InCameraInfos, const FSequenceInfo& InSequenceInfo, TMap<int32, FImgMediaTileSelection>& VisibleTiles) const override
+		void CalculateVisibleTiles(const TArray<FImgMediaViewInfo>& InViewInfos, const FSequenceInfo& InSequenceInfo, TMap<int32, FImgMediaTileSelection>& VisibleTiles) const override
 		{
 			UMeshComponent* Mesh = MeshComponent.Get();
 			if (Mesh == nullptr)
 			{
-				return false;
+				return;
 			}
 
 			// To avoid calculating tile corner mip levels multiple times over, we cache them in this array.
@@ -268,11 +278,16 @@ namespace {
 			FVector TexelOffsetXWS = MeshTransform.TransformVector(FVector(0, PlaneSize.Y / InSequenceInfo.Dim.X, 0));
 			FVector TexelOffsetYWS = MeshTransform.TransformVector(FVector(0, 0, -PlaneSize.Z / InSequenceInfo.Dim.Y));
 
-			for (const FImgMediaMipMapCameraInfo& CameraInfo : InCameraInfos)
+			for (const FImgMediaViewInfo& ViewInfo : InViewInfos)
 			{
+				if (IsPrimitiveComponentHidden(Mesh->ComponentId, ViewInfo))
+				{
+					continue;
+				}
+
 				// Get frustum.
 				FConvexVolume ViewFrustum;
-				GetViewFrustumBounds(ViewFrustum, CameraInfo.ViewProjectionMatrix, false, false);
+				GetViewFrustumBounds(ViewFrustum, ViewInfo.ViewProjectionMatrix, false, false);
 
 				int32 MaxLevel = InSequenceInfo.NumMipLevels - 1;
 				int MipLevelDiv = 1 << MaxLevel;
@@ -341,9 +356,9 @@ namespace {
 									float CornerStepY = TileCornerY / (float)CurrentNumTiles.Y;
 									FVector CornersWS = PlaneCornerWS + (DirXWS * CornerStepX + DirYWS * CornerStepY);
 
-									if (CalculateMipLevel(CameraInfo, CornersWS, TexelOffsetXWS, TexelOffsetYWS, CalculatedLevel))
+									if (CalculateMipLevel(ViewInfo, CornersWS, TexelOffsetXWS, TexelOffsetYWS, CalculatedLevel))
 									{
-										CalculatedLevel += LODBias + CameraInfo.MaterialTextureMipBias;
+										CalculatedLevel += LODBias + ViewInfo.MaterialTextureMipBias;
 
 										SetCachedMipLevel(BaseLevelCorner.X, BaseLevelCorner.Y, MaxCornerX, CalculatedLevel);
 										bValidLevel = true;
@@ -401,7 +416,6 @@ namespace {
 					}
 				}
 			}
-			return true;
 		}
 
 	private:
@@ -439,21 +453,26 @@ namespace {
 	public:
 		using FImgMediaMipMapObjectInfo::FImgMediaMipMapObjectInfo;
 
-		bool CalculateVisibleTiles(const TArray<FImgMediaMipMapCameraInfo>& InCameraInfos, const FSequenceInfo& InSequenceInfo, TMap<int32, FImgMediaTileSelection>& VisibleTiles) const override
+		void CalculateVisibleTiles(const TArray<FImgMediaViewInfo>& InViewInfos, const FSequenceInfo& InSequenceInfo, TMap<int32, FImgMediaTileSelection>& VisibleTiles) const override
 		{
 			UMeshComponent* Mesh = MeshComponent.Get();
 			if (Mesh == nullptr)
 			{
-				return false;
+				return;
 			}
 
 			const float DefaultSphereRadius = 160.0f;
 
-			for (const FImgMediaMipMapCameraInfo& CameraInfo : InCameraInfos)
+			for (const FImgMediaViewInfo& ViewInfo : InViewInfos)
 			{
+				if (IsPrimitiveComponentHidden(Mesh->ComponentId, ViewInfo))
+				{
+					continue;
+				}
+
 				// Analytical derivation of visible tiles from the view frustum, given a sphere presumed to be infinitely large
 				FConvexVolume ViewFrustum;
-				GetViewFrustumBounds(ViewFrustum, CameraInfo.ViewProjectionMatrix, false, false);
+				GetViewFrustumBounds(ViewFrustum, ViewInfo.ViewProjectionMatrix, false, false);
 				
 				// Include all tiles containted in the visible UV region
 				int32 NumX = InSequenceInfo.NumTiles.X;
@@ -515,8 +534,6 @@ namespace {
 					}
 				}
 			}
-
-			return true;
 		}
 	};
 
@@ -644,7 +661,7 @@ void FImgMediaMipMapInfo::UpdateMipLevelCache()
 		// Loop over all objects.
 		for (FImgMediaMipMapObjectInfo* ObjectInfo : Objects)
 		{
-			ObjectInfo->CalculateVisibleTiles(CameraInfos, SequenceInfo, CachedVisibleTiles);
+			ObjectInfo->CalculateVisibleTiles(ViewInfos, SequenceInfo, CachedVisibleTiles);
 		}
 	}
 	
@@ -656,7 +673,11 @@ void FImgMediaMipMapInfo::Tick(float DeltaTime)
 {
 	FScopeLock Lock(&InfoCriticalSection);
 
-	CameraInfos = IImgMediaModule::Get().GetCopyCameraInfos();
+	const TSharedPtr<FImgMediaSceneViewExtension, ESPMode::ThreadSafe>& SVE = IImgMediaModule::Get().GetSceneViewExtension();
+	if (SVE.IsValid())
+	{
+		ViewInfos = SVE->GetViewInfos();
+	}
 
 	// Let the cache update this frame.
 	bIsCacheValid = false;
diff --git a/Engine/Plugins/Media/ImgMedia/Source/ImgMedia/Private/ImgMediaModule.cpp b/Engine/Plugins/Media/ImgMedia/Source/ImgMedia/Private/ImgMediaModule.cpp
--- a/Engine/Plugins/Media/ImgMedia/Source/ImgMedia/Private/ImgMediaModule.cpp
+++ b/Engine/Plugins/Media/ImgMedia/Source/ImgMedia/Private/ImgMediaModule.cpp
@@ -122,14 +122,9 @@ public:
 		return Player;
 	}
 
-	TArray<FImgMediaMipMapCameraInfo> GetCopyCameraInfos() const override
+	const TSharedPtr<FImgMediaSceneViewExtension, ESPMode::ThreadSafe>& GetSceneViewExtension() const override
 	{
-		if (SceneViewExtension.IsValid())
-		{
-			return SceneViewExtension->GetCameraInfos();
-		}
-
-		return {};
+		return SceneViewExtension;
 	}
 
 public:
diff --git a/Engine/Plugins/Media/ImgMedia/Source/ImgMedia/Private/ImgMediaSceneViewExtension.cpp b/Engine/Plugins/Media/ImgMedia/Source/ImgMedia/Private/ImgMediaSceneViewExtension.cpp
--- a/Engine/Plugins/Media/ImgMedia/Source/ImgMedia/Private/ImgMediaSceneViewExtension.cpp
+++ b/Engine/Plugins/Media/ImgMedia/Source/ImgMedia/Private/ImgMediaSceneViewExtension.cpp
@@ -7,7 +7,7 @@
 
 FImgMediaSceneViewExtension::FImgMediaSceneViewExtension(const FAutoRegister& AutoReg)
 	: FSceneViewExtensionBase(AutoReg)
-	, CachedCameraInfos()
+	, CachedViewInfos()
 	, LastFrameNumber(0)
 {
 }
@@ -24,7 +24,7 @@ void FImgMediaSceneViewExtension::BeginRenderViewFamily(FSceneViewFamily& InView
 {
 	if (LastFrameNumber != InViewFamily.FrameNumber)
 	{
-		CachedCameraInfos.Reset();
+		CachedViewInfos.Reset();
 		LastFrameNumber = InViewFamily.FrameNumber;
 	}
 
@@ -40,12 +40,24 @@ void FImgMediaSceneViewExtension::BeginRenderViewFamily(FSceneViewFamily& InView
 
 	for (const FSceneView* View : InViewFamily.Views)
 	{
-		FImgMediaMipMapCameraInfo Info;
+		FImgMediaViewInfo Info;
 		Info.Location = View->ViewMatrices.GetViewOrigin();
 		Info.ViewMatrix = View->ViewMatrices.GetViewMatrix();
 		Info.ViewProjectionMatrix = View->ViewMatrices.GetViewProjectionMatrix();
 		Info.ViewportRect = View->UnconstrainedViewRect.Scale(ResolutionFraction);
 
+		// We store hidden or show-only ids to later avoid needless calculations when objects are not in view.
+		if (View->ShowOnlyPrimitives.IsSet())
+		{
+			Info.bPrimitiveHiddenMode = false;
+			Info.PrimitiveComponentIds = View->ShowOnlyPrimitives.GetValue();
+		}
+		else
+		{
+			Info.bPrimitiveHiddenMode = true;
+			Info.PrimitiveComponentIds = View->HiddenPrimitives;
+		}
+
 		// View->MaterialTextureMipBias is only set later in rendering: we replicate the logic here.
 		if (View->PrimaryScreenPercentageMethod == EPrimaryScreenPercentageMethod::TemporalUpscale)
 		{
@@ -57,6 +69,12 @@ void FImgMediaSceneViewExtension::BeginRenderViewFamily(FSceneViewFamily& InView
 			Info.MaterialTextureMipBias = 0.0f;
 		}
 
-		CachedCameraInfos.Add(MoveTemp(Info));
+		CachedViewInfos.Add(MoveTemp(Info));
 	}
-}
\ No newline at end of file
+}
+
+int32 FImgMediaSceneViewExtension::GetPriority() const
+{
+	// Lowest priority value to ensure all other extensions are executed before ours.
+	return MIN_int32;
+}
diff --git a/Engine/Plugins/Media/ImgMedia/Source/ImgMedia/Private/ImgMediaSceneViewExtension.h b/Engine/Plugins/Media/ImgMedia/Source/ImgMedia/Private/ImgMediaSceneViewExtension.h
deleted file mode 100644
diff --git a/Engine/Plugins/Media/ImgMedia/Source/ImgMedia/Public/IImgMediaModule.h b/Engine/Plugins/Media/ImgMedia/Source/ImgMedia/Public/IImgMediaModule.h
--- a/Engine/Plugins/Media/ImgMedia/Source/ImgMedia/Public/IImgMediaModule.h
+++ b/Engine/Plugins/Media/ImgMedia/Source/ImgMedia/Public/IImgMediaModule.h
@@ -12,6 +12,7 @@
 class FImgMediaGlobalCache;
 struct FImgMediaMipMapCameraInfo;
 class FImgMediaPlayer;
+class FImgMediaSceneViewExtension;
 class IMediaEventSink;
 class IMediaPlayer;
 
@@ -53,12 +54,8 @@ public:
 	 */
 	virtual TSharedPtr<IMediaPlayer, ESPMode::ThreadSafe> CreatePlayer(IMediaEventSink& EventHandler) = 0;
 
-	/**
-	 * Get a copy of camera information cached by our local scene view extension, for mip & tile calculations.
-	 *
-	 * @return Array of info on each camera.
-	 */
-	virtual TArray<FImgMediaMipMapCameraInfo> GetCopyCameraInfos() const = 0;
+	/** Get scene view extension cahcing view information (for mip/tile calculations). */
+	virtual const TSharedPtr<FImgMediaSceneViewExtension, ESPMode::ThreadSafe>& GetSceneViewExtension() const = 0;
 public:
 
 	/** Virtual destructor. */
diff --git a/Engine/Plugins/Media/ImgMedia/Source/ImgMedia/Public/ImgMediaMipMapInfo.h b/Engine/Plugins/Media/ImgMedia/Source/ImgMedia/Public/ImgMediaMipMapInfo.h
--- a/Engine/Plugins/Media/ImgMedia/Source/ImgMedia/Public/ImgMediaMipMapInfo.h
+++ b/Engine/Plugins/Media/ImgMedia/Source/ImgMedia/Public/ImgMediaMipMapInfo.h
@@ -6,6 +6,7 @@
 #include "CoreMinimal.h"
 #include "IMediaOptions.h"
 #include "MediaTextureTracker.h"
+#include "ImgMediaSceneViewExtension.h"
 #include "Tickable.h"
 #include "UObject/WeakObjectPtr.h"
 #include "UObject/WeakObjectPtrTemplates.h"
@@ -38,23 +39,6 @@ struct FSequenceInfo
 	}
 };
 
-/**
- * Holds info on a camera which we can use for mipmap calculations.
- */
-struct FImgMediaMipMapCameraInfo
-{
-	/** Position of camera. */
-	FVector Location;
-	/** View matrix of camera. */
-	FMatrix ViewMatrix;
-	/** View projection matrix of camera. */
-	FMatrix ViewProjectionMatrix;
-	/** Active viewport size. */
-	FIntRect ViewportRect;
-	/** View mip bias. */
-	float MaterialTextureMipBias;
-};
-
 /**
  * Describes which tiles are visible.
  */
@@ -191,7 +175,7 @@ public:
 	 * @param InSequenceInfo Active image sequence information
 	 * @param VisibleTiles Updated list of visible tiles per mip level
 	 */
-	virtual bool CalculateVisibleTiles(const TArray<FImgMediaMipMapCameraInfo>& InCameraInfos, const FSequenceInfo& InSequenceInfo, TMap<int32, FImgMediaTileSelection>& VisibleTiles) const;
+	virtual void CalculateVisibleTiles(const TArray<FImgMediaViewInfo>& InViewInfos, const FSequenceInfo& InSequenceInfo, TMap<int32, FImgMediaTileSelection>& VisibleTiles) const;
 
 	/**
 	 * Get the registered mesh component.
@@ -291,8 +275,8 @@ protected:
 	/** Array of objects that are using our img sequence. */
 	TArray<FImgMediaMipMapObjectInfo*> Objects;
 
-	/** Info for each camera, used in mipmap calculations. */
-	TArray<FImgMediaMipMapCameraInfo> CameraInfos;
+	/** Info for each view, used in mipmap calculations. */
+	TArray<FImgMediaViewInfo> ViewInfos;
 
 	/** Size, tiling and mips sequence information. */
 	FSequenceInfo SequenceInfo;
diff --git a/Engine/Plugins/Media/ImgMedia/Source/ImgMedia/Public/ImgMediaSceneViewExtension.h b/Engine/Plugins/Media/ImgMedia/Source/ImgMedia/Public/ImgMediaSceneViewExtension.h
new file mode 100644
--- /dev/null
+++ b/Engine/Plugins/Media/ImgMedia/Source/ImgMedia/Public/ImgMediaSceneViewExtension.h
@@ -0,0 +1,57 @@
+// Copyright Epic Games, Inc. All Rights Reserved.
+
+#pragma once
+
+#include "CoreMinimal.h"
+#include "Containers/Set.h"
+#include "SceneTypes.h"
+#include "SceneViewExtension.h"
+
+/**
+ * Holds info on a camera which we can use for mipmap calculations.
+ */
+struct IMGMEDIA_API FImgMediaViewInfo
+{
+	/** Position of camera. */
+	FVector Location;
+	/** View matrix of camera. */
+	FMatrix ViewMatrix;
+	/** View projection matrix of camera. */
+	FMatrix ViewProjectionMatrix;
+	/** Active viewport size. */
+	FIntRect ViewportRect;
+	/** View mip bias. */
+	float MaterialTextureMipBias;
+	/** Hidden or show-only mode for primitive components. */
+	bool bPrimitiveHiddenMode;
+	/** Hidden or show-only primitive components. */
+	TSet<FPrimitiveComponentId> PrimitiveComponentIds;
+};
+
+/**
+ * Scene view extension used to cache view information (primarily for visible mip/tile calculations).
+ */
+class FImgMediaSceneViewExtension final : public FSceneViewExtensionBase
+{
+public:
+	FImgMediaSceneViewExtension(const FAutoRegister& AutoReg);
+
+	/**
+	 * Get the cached camera information array, updated on the game thread by BeginRenderViewFamily.
+	 *
+	 * @return Array of info on each camera.
+	 */
+	IMGMEDIA_API const TArray<FImgMediaViewInfo>& GetViewInfos() const { return CachedViewInfos; };
+
+	void SetupViewFamily(FSceneViewFamily& InViewFamily) override;
+	void SetupView(FSceneViewFamily& InViewFamily, FSceneView& InView) override;
+	void BeginRenderViewFamily(FSceneViewFamily& InViewFamily) override;
+	int32 GetPriority() const override;
+
+private:
+	/** Array of info on each camera used for mipmap calculations. */
+	TArray<FImgMediaViewInfo> CachedViewInfos;
+
+	/** Last received FSceneViewFamily frame number. */
+	uint32 LastFrameNumber;
+};
