ImgMedia: Updated sphere computations for the new generated sphere mesh.

#rnx

diff --git a/Engine/Plugins/Media/ImgMedia/Source/ImgMedia/Private/Assets/ImgMediaMipMapInfo.cpp b/Engine/Plugins/Media/ImgMedia/Source/ImgMedia/Private/Assets/ImgMediaMipMapInfo.cpp
--- a/Engine/Plugins/Media/ImgMedia/Source/ImgMedia/Private/Assets/ImgMediaMipMapInfo.cpp
+++ b/Engine/Plugins/Media/ImgMedia/Source/ImgMedia/Private/Assets/ImgMediaMipMapInfo.cpp
@@ -461,7 +461,7 @@
 				return;
 			}
 
-			const float DefaultSphereRadius = 160.0f;
+			const float DefaultSphereRadius = 50.0f;
 
 			for (const FImgMediaViewInfo& ViewInfo : InViewInfos)
 			{
@@ -487,7 +487,7 @@
 						FVector2D TileCornerSpherical = FVector2D(UE_PI * TileCornerUV.Y, UE_TWO_PI * TileCornerUV.X);
 						
 						// Adjust spherical coordinates to default sphere UVs
-						TileCornerSpherical.Y = -UE_HALF_PI - TileCornerSpherical.Y;
+						TileCornerSpherical.Y = -TileCornerSpherical.Y;
 
 						FVector TileCorner = TileCornerSpherical.SphericalToUnitCartesian() * DefaultSphereRadius;
 						TileCorner = Mesh->GetComponentTransform().TransformPosition(TileCorner);
@@ -507,6 +507,17 @@
 							VisibleTiles[0].SetVisible(TileX, AdjacentY);
 							VisibleTiles[0].SetVisible(AdjacentX, AdjacentY);
 						}
+
+#if false
+#if WITH_EDITOR
+						// Enable this to draw a sphere where each tile is.
+						Async(EAsyncExecution::TaskGraphMainThread, [TileCorner]()
+							{
+								UWorld* World = GEditor->GetEditorWorldContext().World();
+								DrawDebugPoint(World, TileCorner, 5.0f, FColor::Red, false, 0.05f);
+							});
+#endif // WITH_EDITOR
+#endif // false
 					}
 				}
 
diff --git a/Engine/Plugins/Media/MediaPlate/Source/MediaPlate/Public/MediaPlateComponent.h b/Engine/Plugins/Media/MediaPlate/Source/MediaPlate/Public/MediaPlateComponent.h
--- a/Engine/Plugins/Media/MediaPlate/Source/MediaPlate/Public/MediaPlateComponent.h
+++ b/Engine/Plugins/Media/MediaPlate/Source/MediaPlate/Public/MediaPlateComponent.h
@@ -97,9 +97,9 @@
 	UPROPERTY(EditAnywhere, AdvancedDisplay, Category = "MediaPlate")
 	float SmartCacheTimeToLookAhead = 0.2f;
 
-	/** Specify type of mesh used for visible mips and tiles calculations.
-	 * Use /MediaPlate/SM_MediaPlateScreen for the plane.
-	 * Use /Engine/BasicShapes/Sphere for the sphere.
+	/**
+	 * Specify type of mesh used for visible mips and tiles calculations.
+	 * (Using the provided plane and sphere meshes only.)
 	 */
 	UPROPERTY(EditAnywhere, AdvancedDisplay, Category = "MediaPlate")
 	EMediaTextureVisibleMipsTiles VisibleMipsTilesCalculations;
