ImgMedia: Privatize mip map info classes and minor cleanup of unused includes.

#rnx

diff --git a/Engine/Plugins/Media/ImgMedia/Source/ImgMedia/Private/Assets/ImgMediaMipMapInfo.h b/Engine/Plugins/Media/ImgMedia/Source/ImgMedia/Private/Assets/ImgMediaMipMapInfo.h
new file mode 100644
--- /dev/null
+++ b/Engine/Plugins/Media/ImgMedia/Source/ImgMedia/Private/Assets/ImgMediaMipMapInfo.h
@@ -0,0 +1,292 @@
+// Copyright Epic Games, Inc. All Rights Reserved.
+
+#pragma once
+
+#include "Containers/BitArray.h"
+#include "CoreMinimal.h"
+#include "IMediaOptions.h"
+#include "MediaTextureTracker.h"
+#include "ImgMediaSceneViewExtension.h"
+#include "Tickable.h"
+#include "UObject/WeakObjectPtr.h"
+#include "UObject/WeakObjectPtrTemplates.h"
+
+class AActor;
+class UMeshComponent;
+class UMediaTexture;
+
+/**
+ * Holds size, tiling and mips sequence information.
+ */
+struct FSequenceInfo
+{
+	/** Name of this sequence. */
+	FName Name;
+	/** Pixel dimensions of this sequence. */
+	FIntPoint Dim;
+	/** Number of tiles in the X, Y directions. */
+	FIntPoint NumTiles;
+	/** Number of mip levels. */
+	int32 NumMipLevels;
+	/** 
+	* Check if the sequence has tiles.
+	* 
+	* @return True if the tile counts are above one.
+	*/
+	FORCEINLINE bool IsTiled() const
+	{
+		return (NumTiles.X > 1) || (NumTiles.Y > 1);
+	}
+};
+
+/**
+ * Describes which tiles are visible.
+ */
+struct FImgMediaTileSelection
+{
+	/**
+	 * Create and initialize a new tile selection.
+	 *
+	 * @param NumTilesX Horizontal tile count.
+	 * @param NumTilesY Vertical tile count.
+	 * @param bDefaultVisibility Optional visibility for the entire region.
+	 */
+	FImgMediaTileSelection(int32 NumTilesX, int32 NumTilesY, bool bDefaultVisibility = false);
+
+	/**
+	 * Create and initialize a new tile selection, adjusting tile counts for the specified higher mip level.
+	 *
+	 * @param BaseNumTilesX Horizontal tile count at mip level 0.
+	 * @param NumTilesY Vertical tile count at mip level 0.
+	 * @param TargetMipLevel Higher target mip level for the selection (usually 1 and above).
+	 * @param bDefaultVisibility Optional visibility for the entire region.
+	 */
+	static FImgMediaTileSelection CreateForTargetMipLevel(int32 BaseNumTilesX, int32 BaseNumTilesY, int32 TargetMipLevel, bool bDefaultVisibility = false);
+
+	FImgMediaTileSelection() = default;
+	~FImgMediaTileSelection() noexcept = default;
+
+	FImgMediaTileSelection(const FImgMediaTileSelection&) = default;
+	FImgMediaTileSelection& operator=(const FImgMediaTileSelection&) = default;
+
+	FImgMediaTileSelection(FImgMediaTileSelection&&) = default;
+	FImgMediaTileSelection& operator=(FImgMediaTileSelection&&) = default;
+
+	/**
+	* See if any tile is visible.
+	*
+	* @return True if so.
+	*/
+	bool IsAnyVisible() const;
+
+	/**
+	 * Check if a specific tile is visible.
+	 *
+	 * @param TileCoordX Horizontal tile coordinate.
+	 * @param TileCoordY Vertical tile coordinate.
+	 * @return True if the coordinate is visible, false otherwise.
+	 */
+	bool IsVisible(int32 TileCoordX, int32 TileCoordY) const;
+
+	/**
+	 * Check if the currently visible tiles are also visible in another selection.
+	 *
+	 * @param Other Selection to compare.
+	 * @return True if the other selection is visible, false otherwise.
+	 */
+	bool Contains(const FImgMediaTileSelection& Other) const;
+
+	/**
+	 * Mark a tile as visible.
+	 *
+	 * @param TileCoordX Horizontal tile coordinate.
+	 * @param TileCoordY Vertical tile coordinate.
+	 */
+	void SetVisible(int32 TileCoordX, int32 TileCoordY);
+
+	/**
+	 * Returns the list of visible tile coordinates, in row and column order.
+	 *
+	 * @return Visible tile coordinates array.
+	 */
+	TArray<FIntPoint> GetVisibleCoordinates() const;
+
+	/**
+	 * Return the rectangular region bounding the visible tiles.
+	 *
+	 * @return Visible tile coordinates array.
+	 */
+	FIntRect GetVisibleRegion() const;
+
+	/**
+	 * Returns the number of visible tiles.
+	 *
+	 * @return Visible tile count.
+	 */
+	int32 NumVisibleTiles() const;
+
+	/**
+	 * Returns the overall dimensions in number of tiles.
+	 *
+	 * @return FIntPoint where X is the horizontal tile count and Y the vertical tile count.
+	 */
+	FIntPoint GetDimensions() const { return Dimensions; }
+
+private:
+	/** Convert from tile coordinates to linear index. */
+	FORCEINLINE static int32 ToIndex(int32 CoordX, int32 CoordY, const FIntPoint& Dim)
+	{
+		return CoordY * Dim.X + CoordX;
+	}
+
+	/** Convert from linear index to tile coordinates. */
+	FORCEINLINE static FIntPoint FromIndex(int32 Index, const FIntPoint& Dim)
+	{
+		return FIntPoint(Index % Dim.X, Index / Dim.X);
+	}
+	
+	/** Bit array of tiles over the entire region's dimensions, true when visible. */
+	TBitArray<> Tiles;
+	
+	/** Number of tiles in each axis. */
+	FIntPoint Dimensions;
+
+	/** Cached visible tile region. */
+	mutable FIntRect CachedVisibleRegion;
+	
+	/** Flag indicating if the cached visible tile region needs to be recalculated. */
+	mutable bool bCachedVisibleRegionDirty;
+};
+
+/**
+ * Describes a single object which is using our img sequence.
+ * Base class for various objects such as planes or spheres.
+ */
+class FImgMediaMipMapObjectInfo
+{
+public:
+	FImgMediaMipMapObjectInfo(UMeshComponent* InMeshComponent, float InLODBias = 0.0f);
+	virtual ~FImgMediaMipMapObjectInfo() = default;
+
+	/**
+	 * Calculate visible tiles per mip level.
+	 *
+	 * @param InCameraInfos Active camera information
+	 * @param InSequenceInfo Active image sequence information
+	 * @param VisibleTiles Updated list of visible tiles per mip level
+	 */
+	virtual void CalculateVisibleTiles(const TArray<FImgMediaViewInfo>& InViewInfos, const FSequenceInfo& InSequenceInfo, TMap<int32, FImgMediaTileSelection>& VisibleTiles) const;
+
+	/**
+	 * Get the registered mesh component.
+	 *
+	 * @return Mesh component
+	 */
+	UMeshComponent* GetMeshComponent() const;
+protected:
+	/** Mesh component onto which the media is displayed. */
+	TWeakObjectPtr<class UMeshComponent> MeshComponent;
+	/** LOD bias for the mipmap level. */
+	float LODBias;
+};
+
+/**
+ * Contains information for working with mip maps.
+ */
+class FImgMediaMipMapInfo : public IMediaOptions::FDataContainer, public FTickableGameObject
+{
+public:
+	FImgMediaMipMapInfo();
+	virtual ~FImgMediaMipMapInfo();
+
+	/**
+	 * This object is using our img sequence.
+	 *
+	 * @param InActor Object using our img sequence.
+	 * @param Width Width of the object. If < 0, then get the width automatically.
+	 */
+	void AddObject(AActor* InActor, float Width, float LODBias, EMediaTextureVisibleMipsTiles MeshType = EMediaTextureVisibleMipsTiles::None);
+
+	/**
+	 * This object is no longer using our img sequence.
+	 *
+	 * @param InActor Object no longer using our img sequence.
+	 */
+	void RemoveObject(AActor* InActor);
+
+	/**
+	 * All the objets that are using this media texture will be used in our mip map calculations.
+	 *
+	 * @param InMediaTexture Media texture to get objects from.
+	 */
+	void AddObjectsUsingThisMediaTexture(UMediaTexture* InMediaTexture);
+
+	/**
+	 * Remove all objects from consideration.
+	 */
+	void ClearAllObjects();
+
+	/**
+	 * Provide information on the texture needed for our image sequence.
+	 *
+	 * @param InSequenceName Name of this sequence.
+	 * @param InNumMipMaps Number of mipmaps in our image sequence.
+	 * @param InNumTiles Number of tiles in our image sequence.
+	 * @param InSequenceDim Dimensions of the textures in our image sequence.
+	 */
+	void SetTextureInfo(FName InSequenceName, int32 InNumMipMaps, const FIntPoint& InNumTiles,
+		const FIntPoint& InSequenceDim);
+
+	/**
+	 * Get what mipmap level should be used.
+	 * Returns the lowest level (highest resolution) mipmap.
+	 * Assumes all higher levels will also be used.
+	 *
+	 * @param TileSelection Will be set with which tiles are visible.
+	 * @return Mipmap level.
+	 */
+	TMap<int32, FImgMediaTileSelection> GetVisibleTiles();
+	
+	/**
+	 * Get information on objects that are using our textures.
+	 * 
+	 * @return Array of objects.
+	 */
+	const TArray<FImgMediaMipMapObjectInfo*>& GetObjects() const { return Objects; }
+
+	/**
+	 * Check if any scene objects are using our img sequence.
+	 *
+	 * @return True if any object is active.
+	 */
+	bool HasObjects() const { return Objects.Num() > 0; }
+
+	//~ FTickableGameObject interface
+	virtual void Tick(float DeltaTime) override;
+	virtual TStatId GetStatId() const override { RETURN_QUICK_DECLARE_CYCLE_STAT(FImgMediaMipMapInfo, STATGROUP_Tickables); }
+	virtual bool IsTickableInEditor() const override { return true; }
+
+protected:
+	/**
+	 * Performs mipmap calculations and caches the data.
+	 */
+	void UpdateMipLevelCache();
+
+	/** Array of objects that are using our img sequence. */
+	TArray<FImgMediaMipMapObjectInfo*> Objects;
+
+	/** Info for each view, used in mipmap calculations. */
+	TArray<FImgMediaViewInfo> ViewInfos;
+
+	/** Size, tiling and mips sequence information. */
+	FSequenceInfo SequenceInfo;
+	
+	/** Protects info variables. */
+	FCriticalSection InfoCriticalSection;
+
+	/** Desired mipmap level at this current time. */
+	TMap<int32, FImgMediaTileSelection> CachedVisibleTiles;
+
+	/** True if the cached mipmap data has been calculated this frame. */
+	bool bIsCacheValid;
+};
diff --git a/Engine/Plugins/Media/ImgMedia/Source/ImgMedia/Private/Loader/ImgMediaLoader.cpp b/Engine/Plugins/Media/ImgMedia/Source/ImgMedia/Private/Loader/ImgMediaLoader.cpp
--- a/Engine/Plugins/Media/ImgMedia/Source/ImgMedia/Private/Loader/ImgMediaLoader.cpp
+++ b/Engine/Plugins/Media/ImgMedia/Source/ImgMedia/Private/Loader/ImgMediaLoader.cpp
@@ -20,7 +20,6 @@
 #include "IImgMediaModule.h"
 #include "IImgMediaReader.h"
 #include "ImgMediaLoaderWork.h"
-#include "ImgMediaMipMapInfo.h"
 #include "ImgMediaScheduler.h"
 #include "ImgMediaTextureSample.h"
 
diff --git a/Engine/Plugins/Media/ImgMedia/Source/ImgMedia/Private/Loader/ImgMediaLoader.h b/Engine/Plugins/Media/ImgMedia/Source/ImgMedia/Private/Loader/ImgMediaLoader.h
--- a/Engine/Plugins/Media/ImgMedia/Source/ImgMedia/Private/Loader/ImgMediaLoader.h
+++ b/Engine/Plugins/Media/ImgMedia/Source/ImgMedia/Private/Loader/ImgMediaLoader.h
@@ -8,13 +8,12 @@
 #include "IMediaSamples.h"
 #include "IMediaTextureSample.h"
 #include "IMediaView.h"
-#include "ImgMediaMipMapInfo.h"
+#include "Assets/ImgMediaMipMapInfo.h"
 #include "Misc/FrameRate.h"
 #include "Templates/SharedPointer.h"
 
 class FImgMediaGlobalCache;
 class FImgMediaLoaderWork;
-class FImgMediaMipMapInfo;
 class FImgMediaScheduler;
 class FImgMediaTextureSample;
 class FMediaTimeStamp;
diff --git a/Engine/Plugins/Media/ImgMedia/Source/ImgMedia/Private/Readers/ExrImgMediaReaderGpu.cpp b/Engine/Plugins/Media/ImgMedia/Source/ImgMedia/Private/Readers/ExrImgMediaReaderGpu.cpp
--- a/Engine/Plugins/Media/ImgMedia/Source/ImgMedia/Private/Readers/ExrImgMediaReaderGpu.cpp
+++ b/Engine/Plugins/Media/ImgMedia/Source/ImgMedia/Private/Readers/ExrImgMediaReaderGpu.cpp
@@ -12,7 +12,6 @@
 #include "ExrReaderGpu.h"
 #include "RHICommandList.h"
 #include "ExrSwizzlingShader.h"
-#include "ImgMedia/Public/ImgMediaMipMapInfo.h"
 #include "SceneUtils.h"
 
 
diff --git a/Engine/Plugins/Media/ImgMedia/Source/ImgMedia/Public/ImgMediaMipMapInfo.h b/Engine/Plugins/Media/ImgMedia/Source/ImgMedia/Public/ImgMediaMipMapInfo.h
deleted file mode 100644
diff --git a/Engine/Plugins/Media/ImgMedia/Source/ImgMediaEditor/Private/Customizations/ImgMediaSourceCustomization.cpp b/Engine/Plugins/Media/ImgMedia/Source/ImgMediaEditor/Private/Customizations/ImgMediaSourceCustomization.cpp
--- a/Engine/Plugins/Media/ImgMedia/Source/ImgMediaEditor/Private/Customizations/ImgMediaSourceCustomization.cpp
+++ b/Engine/Plugins/Media/ImgMedia/Source/ImgMediaEditor/Private/Customizations/ImgMediaSourceCustomization.cpp
@@ -4,7 +4,6 @@
 
 #include "IMediaModule.h"
 #include "ImgMediaEditorModule.h"
-#include "ImgMediaMipMapInfo.h"
 #include "ImgMediaSource.h"
 #include "GameFramework/Actor.h"
 #include "Misc/Paths.h"
diff --git a/Engine/Plugins/Media/ImgMedia/Source/ImgMediaEditor/Private/Customizations/ImgMediaSourceCustomization.h b/Engine/Plugins/Media/ImgMedia/Source/ImgMediaEditor/Private/Customizations/ImgMediaSourceCustomization.h
--- a/Engine/Plugins/Media/ImgMedia/Source/ImgMediaEditor/Private/Customizations/ImgMediaSourceCustomization.h
+++ b/Engine/Plugins/Media/ImgMedia/Source/ImgMediaEditor/Private/Customizations/ImgMediaSourceCustomization.h
@@ -14,9 +14,6 @@
 class IPropertyHandle;
 class SEditableTextBox;
 
-class FImgMediaMipMapInfo;
-struct FImgMediaMipMapCameraInfo;
-
 /**
  * Implements a details view customization for the UImgMediaSource class.
  */
diff --git a/Engine/Plugins/Media/ImgMedia/Source/ImgMediaEngine/Private/Unreal/ImgMediaPlaybackComponent.cpp b/Engine/Plugins/Media/ImgMedia/Source/ImgMediaEngine/Private/Unreal/ImgMediaPlaybackComponent.cpp
--- a/Engine/Plugins/Media/ImgMedia/Source/ImgMediaEngine/Private/Unreal/ImgMediaPlaybackComponent.cpp
+++ b/Engine/Plugins/Media/ImgMedia/Source/ImgMediaEngine/Private/Unreal/ImgMediaPlaybackComponent.cpp
@@ -3,7 +3,6 @@
 #include "Unreal/ImgMediaPlaybackComponent.h"
 
 #include "Components/PrimitiveComponent.h"
-#include "ImgMediaMipMapInfo.h"
 #include "Materials/MaterialInterface.h"
 #include "MediaTexture.h"
 #include "MediaTextureTracker.h"
