ImgMedia: Fixed compile issue with process images player code.
#rnx

diff --git a/Engine/Plugins/Media/ImgMedia/Source/ImgMediaEditor/Private/Widgets/SImgMediaProcessImages.cpp b/Engine/Plugins/Media/ImgMedia/Source/ImgMediaEditor/Private/Widgets/SImgMediaProcessImages.cpp
--- a/Engine/Plugins/Media/ImgMedia/Source/ImgMediaEditor/Private/Widgets/SImgMediaProcessImages.cpp
+++ b/Engine/Plugins/Media/ImgMedia/Source/ImgMediaEditor/Private/Widgets/SImgMediaProcessImages.cpp
@@ -152,7 +152,7 @@ FReply SImgMediaProcessImages::OnProcessImagesClicked()
 		MediaTexture->AddToRoot();
 
 		// Create media source.
-		MediaSource = UMediaSource::SpawnMediaSourceForString(Options->SequencePath.FilePath);
+		MediaSource = UMediaSource::SpawnMediaSourceForString(Options->SequencePath.FilePath, GetTransientPackage());
 		if (MediaSource == nullptr)
 		{
 			return FReply::Handled();
