ImgMedia: ProcessImages now has an option in the editor to use a player, instead of it being a #define.
#rnx

diff --git a/Engine/Plugins/Media/ImgMedia/Source/ImgMediaEditor/Private/Widgets/ImgMediaProcessImagesOptions.h b/Engine/Plugins/Media/ImgMedia/Source/ImgMediaEditor/Private/Widgets/ImgMediaProcessImagesOptions.h
--- a/Engine/Plugins/Media/ImgMedia/Source/ImgMediaEditor/Private/Widgets/ImgMediaProcessImagesOptions.h
+++ b/Engine/Plugins/Media/ImgMedia/Source/ImgMediaEditor/Private/Widgets/ImgMediaProcessImagesOptions.h
@@ -47,6 +47,10 @@ public:
 	UPROPERTY(EditAnywhere, Transient, Category = Processing, meta = (ClampMin = "0.0"))
 	int32 NumThreads = 0;
 
+	/** Use a player to read in the image. */
+	UPROPERTY(EditAnywhere, Transient, Category = Processing)
+	bool bUsePlayer = false;
+
 	/** Number of tiles in the X direction. If 0, then there are no tiles. */
 	UPROPERTY(VisibleAnywhere, Transient, Category = Tiles)
 	int32 NumTilesX = 0;
diff --git a/Engine/Plugins/Media/ImgMedia/Source/ImgMediaEditor/Private/Widgets/SImgMediaProcessImages.cpp b/Engine/Plugins/Media/ImgMedia/Source/ImgMediaEditor/Private/Widgets/SImgMediaProcessImages.cpp
--- a/Engine/Plugins/Media/ImgMedia/Source/ImgMediaEditor/Private/Widgets/SImgMediaProcessImages.cpp
+++ b/Engine/Plugins/Media/ImgMedia/Source/ImgMediaEditor/Private/Widgets/SImgMediaProcessImages.cpp
@@ -40,8 +40,6 @@
 
 #define LOCTEXT_NAMESPACE "ImgMediaProcessImages"
 
-#define IMGMEDIA_PROCESSIMAGES_USE_PLAYER 0
-
 SImgMediaProcessImages::~SImgMediaProcessImages()
 {
 	CleanUp();
@@ -113,9 +111,11 @@ END_SLATE_FUNCTION_BUILD_OPTIMIZATION
 void SImgMediaProcessImages::Tick(const FGeometry& AllottedGeometry, const double InCurrentTime, const float InDeltaTime)
 {
 	SCompoundWidget::Tick(AllottedGeometry, InCurrentTime, InDeltaTime);
-#if IMGMEDIA_PROCESSIMAGES_USE_PLAYER
-	HandleProcessing();
-#endif
+
+	if (bUsePlayer)
+	{
+		HandleProcessing();
+	}
 }
 
 void SImgMediaProcessImages::UpdateWidgets()
@@ -130,6 +130,7 @@ FReply SImgMediaProcessImages::OnProcessImagesClicked()
 	{
 		// Set that we are processing now.
 		bIsProcessing = true;
+		bUsePlayer = Options->bUsePlayer;
 		UpdateWidgets();
 
 		// Create notification.
@@ -138,41 +139,44 @@ FReply SImgMediaProcessImages::OnProcessImagesClicked()
 
 		ConfirmNotification = FSlateNotificationManager::Get().AddNotification(Info);
 
-#if IMGMEDIA_PROCESSIMAGES_USE_PLAYER
-		// Create player.
-		MediaPlayer = NewObject<UMediaPlayer>(GetTransientPackage(), "MediaPlayer", RF_Transient);
-		MediaPlayer->SetLooping(true);
-		MediaPlayer->PlayOnOpen = true;
-		MediaPlayer->AddToRoot();
-
-		// Create texture.
-		MediaTexture = NewObject<UMediaTexture>(GetTransientPackage(), "MediaTexture", RF_Transient);
-		MediaTexture->SetMediaPlayer(MediaPlayer);
-		MediaTexture->UpdateResource();
-		MediaTexture->AddToRoot();
-
-		// Create media source.
-		MediaSource = UMediaSource::SpawnMediaSourceForString(Options->SequencePath.FilePath, GetTransientPackage());
-		if (MediaSource == nullptr)
+		// Are we using a player?
+		if (bUsePlayer)
 		{
-			return FReply::Handled();
+			// Create player.
+			MediaPlayer = NewObject<UMediaPlayer>(GetTransientPackage(), "MediaPlayer", RF_Transient);
+			MediaPlayer->SetLooping(true);
+			MediaPlayer->PlayOnOpen = true;
+			MediaPlayer->AddToRoot();
+
+			// Create texture.
+			MediaTexture = NewObject<UMediaTexture>(GetTransientPackage(), "MediaTexture", RF_Transient);
+			MediaTexture->SetMediaPlayer(MediaPlayer);
+			MediaTexture->UpdateResource();
+			MediaTexture->AddToRoot();
+
+			// Create media source.
+			MediaSource = UMediaSource::SpawnMediaSourceForString(Options->SequencePath.FilePath, GetTransientPackage());
+			if (MediaSource == nullptr)
+			{
+				return FReply::Handled();
+			}
+			MediaSource->AddToRoot();
+
+			// Start playing.
+			CurrentFrameIndex = 0;
+			CurrentTime = FTimespan::FromSeconds(0.0f);
+			MediaPlayer->SetBlockOnTimeRange(TRange<FTimespan>(CurrentTime,
+				CurrentTime + FTimespan::FromSeconds(1.0f / 100000.0f)));
+			MediaPlayer->OpenSource(MediaSource);
 		}
-		MediaSource->AddToRoot();
-
-		// Start playing.
-		CurrentFrameIndex = 0;
-		CurrentTime = FTimespan::FromSeconds(0.0f);
-		MediaPlayer->SetBlockOnTimeRange(TRange<FTimespan>(CurrentTime,
-			CurrentTime + FTimespan::FromSeconds(1.0f /100000.0f)));
-		MediaPlayer->OpenSource(MediaSource);
-#else
-		
-		// Start async task to process files.
-		Async(EAsyncExecution::Thread, [this]()
+		else
 		{
-			ProcessAllImages();
-		});
-#endif
+			// Start async task to process files.
+			Async(EAsyncExecution::Thread, [this]()
+			{
+				ProcessAllImages();
+			});
+		}
 	}
 
 	return FReply::Handled();
diff --git a/Engine/Plugins/Media/ImgMedia/Source/ImgMediaEditor/Private/Widgets/SImgMediaProcessImages.h b/Engine/Plugins/Media/ImgMedia/Source/ImgMediaEditor/Private/Widgets/SImgMediaProcessImages.h
--- a/Engine/Plugins/Media/ImgMedia/Source/ImgMediaEditor/Private/Widgets/SImgMediaProcessImages.h
+++ b/Engine/Plugins/Media/ImgMedia/Source/ImgMediaEditor/Private/Widgets/SImgMediaProcessImages.h
@@ -206,5 +206,7 @@ private:
 	bool bIsProcessing = false;
 	/** True if we want to cancel processing. */
 	bool bIsCancelling = false;
+	/** True if we are using a player to read in the images. */
+	bool bUsePlayer = false;
 	
 };
