ImgMedia: Fixed process images sequence path file picker not showing the directory you just selected.
#rnx

diff --git a/Engine/Plugins/Media/ImgMedia/Source/ImgMediaEditor/Private/Customizations/ImgMediaFilePathCustomization.cpp b/Engine/Plugins/Media/ImgMedia/Source/ImgMediaEditor/Private/Customizations/ImgMediaFilePathCustomization.cpp
--- a/Engine/Plugins/Media/ImgMedia/Source/ImgMediaEditor/Private/Customizations/ImgMediaFilePathCustomization.cpp
+++ b/Engine/Plugins/Media/ImgMedia/Source/ImgMediaEditor/Private/Customizations/ImgMediaFilePathCustomization.cpp
@@ -59,7 +59,10 @@ void FImgMediaFilePathCustomization::CustomizeChildren(TSharedRef<IPropertyHandl
 				.BrowseButtonImage(FAppStyle::GetBrush("PropertyWindow.Button_Ellipsis"))
 				.BrowseButtonStyle(FAppStyle::Get(), "HoverHintOnly")
 				.BrowseButtonToolTip(LOCTEXT("FileButtonToolTipText", "Choose a file from this computer"))
-				.BrowseDirectory(FPaths::GetPath(PickedPath))
+				.BrowseDirectory_Lambda([this]()->FString
+				{
+					return FPaths::GetPath(HandleFilePathPickerFilePath());
+				})
 				.BrowseTitle(LOCTEXT("PropertyEditorTitle", "File picker..."))
 				.FilePath(this, &FImgMediaFilePathCustomization::HandleFilePathPickerFilePath)
 				.FileTypeFilter(FileTypeFilter)
