ImgMedia: Added function to calculate regions of visible tiles (mainly for use by the readers where batching can improve performance).

#rnx

diff --git a/Engine/Plugins/Media/ImgMedia/Source/ImgMedia/Private/Assets/ImgMediaMipMapInfo.cpp b/Engine/Plugins/Media/ImgMedia/Source/ImgMedia/Private/Assets/ImgMediaMipMapInfo.cpp
--- a/Engine/Plugins/Media/ImgMedia/Source/ImgMedia/Private/Assets/ImgMediaMipMapInfo.cpp
+++ b/Engine/Plugins/Media/ImgMedia/Source/ImgMedia/Private/Assets/ImgMediaMipMapInfo.cpp
@@ -110,6 +110,70 @@
 	return OutCoordinates;
 }
 
+TArray<FIntRect> FImgMediaTileSelection::GetVisibleRegions() const
+{
+	/**
+	 * This is a two-pass algorithm to batch tiles into contiguous regions, with a bias for row groupings.
+	 * First, we iterate through visible tiles, and create regions for (horizontally) contiguous tiles in each row.
+	 * Second, we create the final regions out of (vertically) contiguous row regions of matching width & position.
+	*/
+
+	TArray<TArray<FIntRect>> RowsOfRegions;
+
+	for (int32 CoordY = 0; CoordY < Dimensions.Y; ++CoordY)
+	{
+		TArray<FIntRect> RegionsPerRow;
+
+		for (int32 CoordX = 0; CoordX < Dimensions.X; ++CoordX)
+		{
+			if (Tiles[ToIndex(CoordX, CoordY, Dimensions)])
+			{
+				FIntPoint TileCoord(CoordX, CoordY);
+
+				bool bIsPreviousRowTileVisible = (CoordX > 0) ? Tiles[ToIndex(CoordX - 1, CoordY, Dimensions)] : false;
+				if (bIsPreviousRowTileVisible)
+				{
+					RegionsPerRow.Last().Include(TileCoord + 1);
+				}
+				else
+				{
+					RegionsPerRow.Emplace(TileCoord, TileCoord + 1);
+				}
+			}
+		}
+
+		if (RegionsPerRow.Num() > 0)
+		{
+			RowsOfRegions.Add(MoveTemp(RegionsPerRow));
+		}
+	}
+
+	TArray<FIntRect> FinalRegions;
+
+	for (const TArray<FIntRect>& RegionsPerRow : RowsOfRegions)
+	{
+		for (const FIntRect& Region : RegionsPerRow)
+		{
+			FIntRect* ContiguousRegion = FinalRegions.FindByPredicate([&Region](const FIntRect& BatchedRegion)
+				{
+					// Batch row regions if their width matches and if they are vertically contiguous.
+					return (Region.Min.X == BatchedRegion.Min.X) && (Region.Max.X == BatchedRegion.Max.X) && (Region.Min.Y == BatchedRegion.Max.Y);
+				});
+
+			if (ContiguousRegion != nullptr)
+			{
+				ContiguousRegion->Max.Y++;
+			}
+			else
+			{
+				FinalRegions.Add(Region);
+			}
+		}
+	}
+
+	return FinalRegions;
+}
+
 FIntRect FImgMediaTileSelection::GetVisibleRegion() const
 {
 	// We offload the region calculation to the loader workers, instead of constantly updating it during SetVisible().
diff --git a/Engine/Plugins/Media/ImgMedia/Source/ImgMedia/Private/Assets/ImgMediaMipMapInfo.h b/Engine/Plugins/Media/ImgMedia/Source/ImgMedia/Private/Assets/ImgMediaMipMapInfo.h
--- a/Engine/Plugins/Media/ImgMedia/Source/ImgMedia/Private/Assets/ImgMediaMipMapInfo.h
+++ b/Engine/Plugins/Media/ImgMedia/Source/ImgMedia/Private/Assets/ImgMediaMipMapInfo.h
@@ -112,6 +112,13 @@
 	TArray<FIntPoint> GetVisibleCoordinates() const;
 
 	/**
+	 * Returns a calculated list of contiguous visible tile regions.
+	 *
+	 * @return Visible tile regions array.
+	 */
+	TArray<FIntRect> GetVisibleRegions() const;
+
+	/**
 	 * Return the rectangular region bounding the visible tiles.
 	 *
 	 * @return Visible tile coordinates array.
diff --git a/Engine/Plugins/Media/ImgMedia/Source/ImgMedia/Private/Tests/ImgMediaTests.cpp b/Engine/Plugins/Media/ImgMedia/Source/ImgMedia/Private/Tests/ImgMediaTests.cpp
new file mode 100644
--- /dev/null
+++ b/Engine/Plugins/Media/ImgMedia/Source/ImgMedia/Private/Tests/ImgMediaTests.cpp
@@ -0,0 +1,92 @@
+// Copyright Epic Games, Inc. All Rights Reserved.
+
+#include "CoreTypes.h"
+#include "Containers/UnrealString.h"
+#include "Logging/LogMacros.h"
+#include "Misc/AutomationTest.h"
+
+#include "ImgMediaMipMapInfo.h"
+
+#if WITH_DEV_AUTOMATION_TESTS
+
+IMPLEMENT_SIMPLE_AUTOMATION_TEST(FImgMediaTests, "System.Plugins.ImgMedia.TileSelection", EAutomationTestFlags::EditorContext | EAutomationTestFlags::EngineFilter)
+bool FImgMediaTests::RunTest(const FString& Parameters)
+{
+	{
+		// Single continuous region
+		FImgMediaTileSelection Selection = FImgMediaTileSelection(10, 10);
+		for (int32 Ty = 3; Ty < 7; Ty++)
+		{
+			for (int32 Tx = 1; Tx < 6; Tx++)
+			{
+				Selection.SetVisible(Tx, Ty);
+			}
+		}
+
+		FIntRect R0 = Selection.GetVisibleRegion();
+		FIntRect R1 = Selection.GetVisibleRegions()[0];
+		AddErrorIfFalse(R0 == R1, TEXT("FImgMediaTests: Mismatched tile regions."));
+	}
+
+	{
+		//Two regions on both sides
+		FImgMediaTileSelection Selection = FImgMediaTileSelection(10, 10);
+		for (int32 Ty = 0; Ty < 10; Ty++)
+		{
+			for (int32 Tx = 0; Tx < 2; Tx++)
+			{
+				Selection.SetVisible(Tx, Ty);
+			}
+
+			for (int32 Tx = 8; Tx < 10; Tx++)
+			{
+				Selection.SetVisible(Tx, Ty);
+			}
+		}
+		TArray<FIntRect> Result = Selection.GetVisibleRegions();
+
+		bool bExpectedNum = Result.Num() == 2;
+		if (bExpectedNum)
+		{
+			AddErrorIfFalse(Result[0] == FIntRect(FIntPoint(0, 0), FIntPoint(2, 10)), TEXT("FImgMediaTests: Mismatched tile regions."));
+			AddErrorIfFalse(Result[1] == FIntRect(FIntPoint(8, 0), FIntPoint(10, 10)), TEXT("FImgMediaTests: Mismatched tile regions."));
+		}
+		else
+		{
+			AddError(TEXT("FImgMediaTests: Expected 2 regions."));
+		}
+	}
+
+	{
+		//Each row has a different length, resulting in one region for each row with the current algo.
+		FImgMediaTileSelection Selection = FImgMediaTileSelection(10, 10);
+		for (int32 Ty = 0; Ty < 10; Ty++)
+		{
+			for (int32 Tx = 0; Tx < Ty + 1; Tx++)
+			{
+				Selection.SetVisible(Tx, Ty);
+			}
+		}
+
+		TArray<FIntRect> Result = Selection.GetVisibleRegions();
+
+		AddErrorIfFalse(Result.Num() == 10, TEXT("FImgMediaTests: Expected 10 regions."));
+	}
+
+	{
+		//Worst case: checkerboard pattern where each tile becomes a region.
+		FImgMediaTileSelection Selection = FImgMediaTileSelection(4, 4);
+		Selection.SetVisible(0, 0); Selection.SetVisible(2, 0);
+		Selection.SetVisible(1, 1); Selection.SetVisible(3, 1);
+		Selection.SetVisible(0, 2); Selection.SetVisible(2, 2);
+		Selection.SetVisible(1, 3); Selection.SetVisible(3, 3);
+
+		TArray<FIntRect> Result = Selection.GetVisibleRegions();
+
+		AddErrorIfFalse(Result.Num() == 8, TEXT("FImgMediaTests: Expected 8 regions."));
+	}
+
+	return !HasAnyErrors();
+}
+
+#endif
