ImgMedia: Visible tile calculation bug fix.

diff --git a/Engine/Plugins/Media/ImgMedia/Source/ImgMedia/Private/Assets/ImgMediaMipMapInfo.cpp b/Engine/Plugins/Media/ImgMedia/Source/ImgMedia/Private/Assets/ImgMediaMipMapInfo.cpp
--- a/Engine/Plugins/Media/ImgMedia/Source/ImgMedia/Private/Assets/ImgMediaMipMapInfo.cpp
+++ b/Engine/Plugins/Media/ImgMedia/Source/ImgMedia/Private/Assets/ImgMediaMipMapInfo.cpp
@@ -382,6 +382,12 @@ namespace {
 					CurrentNumTiles.X = FMath::Max(1, FMath::CeilToInt(float(InSequenceInfo.NumTiles.X) / MipLevelDiv));
 					CurrentNumTiles.Y = FMath::Max(1, FMath::CeilToInt(float(InSequenceInfo.NumTiles.Y) / MipLevelDiv));
 
+					// Exclude subdivided tiles (enqueued below) that are not present (i.e. mipped sequences with odd number of tiles)
+					if (Tile.X >= CurrentNumTiles.X || Tile.Y >= CurrentNumTiles.Y)
+					{
+						continue;
+					}
+
 					// Calculate the tile location in world-space
 					float StepX = float(Tile.X + 0.5f) / CurrentNumTiles.X;
 					float StepY = float(Tile.Y + 0.5f) / CurrentNumTiles.Y;
