ImgMedia: Fixed crash when converting a 1080p image with tiles of size 256.

#rnx

diff --git a/Engine/Plugins/Media/ImgMedia/Source/ImgMediaEditor/Private/Widgets/SImgMediaProcessImages.cpp b/Engine/Plugins/Media/ImgMedia/Source/ImgMediaEditor/Private/Widgets/SImgMediaProcessImages.cpp
--- a/Engine/Plugins/Media/ImgMedia/Source/ImgMediaEditor/Private/Widgets/SImgMediaProcessImages.cpp
+++ b/Engine/Plugins/Media/ImgMedia/Source/ImgMediaEditor/Private/Widgets/SImgMediaProcessImages.cpp
@@ -823,7 +823,7 @@ void SImgMediaProcessImages::TileData(uint8* SourceData, TArray64<uint8>& DestAr
 			int32 ThisDestTileHeight = DestTileHeight;
 			if ((TileY + 1) * TileHeight > SourceHeight)
 			{
-				ThisDestTileHeight -= SourceHeight - TileY * TileHeight;
+				ThisDestTileHeight = SourceHeight - TileY * TileHeight;
 			}
 
 			// Loop over each row in the tile.
