ImgMedia: Update process images tile sizes to meaningful default values.

#rnx

diff --git a/Engine/Plugins/Media/ImgMedia/Source/ImgMediaEditor/Private/Widgets/ImgMediaProcessImagesOptions.h b/Engine/Plugins/Media/ImgMedia/Source/ImgMediaEditor/Private/Widgets/ImgMediaProcessImagesOptions.h
--- a/Engine/Plugins/Media/ImgMedia/Source/ImgMediaEditor/Private/Widgets/ImgMediaProcessImagesOptions.h
+++ b/Engine/Plugins/Media/ImgMedia/Source/ImgMediaEditor/Private/Widgets/ImgMediaProcessImagesOptions.h
@@ -30,17 +30,21 @@ public:
 	/** If checked, then enable mip mapping. */
 	UPROPERTY(EditAnywhere, Category = Sequence)
 	bool bEnableMipMapping;
+
+	/** If checked, then enable tiling. */
+	UPROPERTY(EditAnywhere, Category = Tiles)
+	bool bEnableTiling;
 	
 	/** Width of a tile in pixels. If 0, then do not make tiles. */
-	UPROPERTY(EditAnywhere, Transient, Category = Tiles, meta = (ClampMin = "0.0"))
-	int32 TileSizeX = 0;
+	UPROPERTY(EditAnywhere, Transient, Category = Tiles, meta = (EditCondition = "bEnableTiling", ClampMin = "0.0"))
+	int32 TileSizeX = 256;
 
 	/** Height of a tile in pixels. If 0, then do not make tiles. */
-	UPROPERTY(EditAnywhere, Transient, Category = Tiles, meta = (ClampMin = "0.0"))
-	int32 TileSizeY = 0;
+	UPROPERTY(EditAnywhere, Transient, Category = Tiles, meta = (EditCondition = "bEnableTiling", ClampMin = "0.0"))
+	int32 TileSizeY = 256;
 
 	/** Number of pixels to duplicate along each tile border. */
-	UPROPERTY(EditAnywhere, Transient, Category = Tiles, meta = (ClampMin = "0.0"))
+	UPROPERTY(EditAnywhere, Transient, Category = Tiles, meta = (EditCondition = "bEnableTiling", ClampMin = "0.0"))
 	int32 TileBorder = 0;
 
 	/** Number of threads to use when processing. If 0 then this number is set automatically. */
diff --git a/Engine/Plugins/Media/ImgMedia/Source/ImgMediaEditor/Private/Widgets/SImgMediaProcessImages.cpp b/Engine/Plugins/Media/ImgMedia/Source/ImgMediaEditor/Private/Widgets/SImgMediaProcessImages.cpp
--- a/Engine/Plugins/Media/ImgMedia/Source/ImgMediaEditor/Private/Widgets/SImgMediaProcessImages.cpp
+++ b/Engine/Plugins/Media/ImgMedia/Source/ImgMediaEditor/Private/Widgets/SImgMediaProcessImages.cpp
@@ -198,9 +198,9 @@ void SImgMediaProcessImages::ProcessAllImages()
 	TRACE_CPUPROFILER_EVENT_SCOPE(SImgMediaProcessImages::ProcessAllImages);
 
 	bool bUseCustomFormat = Options->bUseCustomFormat;
-	int32 InTileWidth = Options->TileSizeX;
-	int32 InTileHeight = Options->TileSizeY;
-	int32 TileBorder = Options->TileBorder;
+	int32 InTileWidth = Options->bEnableTiling ? Options->TileSizeX : 0;
+	int32 InTileHeight = Options->bEnableTiling ? Options->TileSizeY : 0;
+	int32 TileBorder = Options->bEnableTiling ? Options->TileBorder : 0;
 	bool bEnableMips = Options->bEnableMipMapping;
 
 	// Create output directory.
@@ -910,9 +910,9 @@ void SImgMediaProcessImages::HandleProcessing()
 					int32 Height = RenderTarget->GetSurfaceHeight();
 					int32 BitDepth = 16;
 					bool bUseCustomFormat = Options->bUseCustomFormat;
-					int32 InTileWidth = Options->TileSizeX;
-					int32 InTileHeight = Options->TileSizeY;
-					int32 TileBorder = Options->TileBorder;
+					int32 InTileWidth = Options->bEnableTiling ? Options->TileSizeX : 0;
+					int32 InTileHeight = Options->bEnableTiling ? Options->TileSizeY : 0;
+					int32 TileBorder = Options->bEnableTiling ? Options->TileBorder : 0;
 					bool bEnableMips = Options->bEnableMipMapping;
 					bool bHasAlphaChannel = false;
 					FString OutPath = Options->OutputPath.Path;
