ImgMedia: Mipmap info tracked objects removed on close, instead of removing all on play. This fixes mip/tile calculation support for duplicate media sources (given unique media textures for media plates).

diff --git a/Engine/Plugins/Media/ImgMedia/Source/ImgMedia/Private/Assets/ImgMediaMipMapInfo.cpp b/Engine/Plugins/Media/ImgMedia/Source/ImgMedia/Private/Assets/ImgMediaMipMapInfo.cpp
--- a/Engine/Plugins/Media/ImgMedia/Source/ImgMedia/Private/Assets/ImgMediaMipMapInfo.cpp
+++ b/Engine/Plugins/Media/ImgMedia/Source/ImgMedia/Private/Assets/ImgMediaMipMapInfo.cpp
@@ -692,6 +692,24 @@ void FImgMediaMipMapInfo::AddObjectsUsingThisMediaTexture(UMediaTexture* InMedia
 	}
 }
 
+void FImgMediaMipMapInfo::RemoveObjectsUsingThisMediaTexture(UMediaTexture* InMediaTexture)
+{
+	// Get objects using this texture.
+	FMediaTextureTracker& TextureTracker = FMediaTextureTracker::Get();
+	const TArray<TWeakPtr<FMediaTextureTrackerObject, ESPMode::ThreadSafe>>* ObjectInfos = TextureTracker.GetObjects(InMediaTexture);
+	if (ObjectInfos != nullptr)
+	{
+		for (TWeakPtr<FMediaTextureTrackerObject, ESPMode::ThreadSafe> ObjectInfoPtr : *ObjectInfos)
+		{
+			TSharedPtr<FMediaTextureTrackerObject, ESPMode::ThreadSafe> ObjectInfo = ObjectInfoPtr.Pin();
+			if (ObjectInfo.IsValid())
+			{
+				RemoveObject(ObjectInfo->Object.Get());
+			}
+		}
+	}
+}
+
 void FImgMediaMipMapInfo::ClearAllObjects()
 {
 	for (FImgMediaMipMapObjectInfo* Info : Objects)
diff --git a/Engine/Plugins/Media/ImgMedia/Source/ImgMedia/Private/Assets/ImgMediaMipMapInfo.h b/Engine/Plugins/Media/ImgMedia/Source/ImgMedia/Private/Assets/ImgMediaMipMapInfo.h
--- a/Engine/Plugins/Media/ImgMedia/Source/ImgMedia/Private/Assets/ImgMediaMipMapInfo.h
+++ b/Engine/Plugins/Media/ImgMedia/Source/ImgMedia/Private/Assets/ImgMediaMipMapInfo.h
@@ -222,12 +222,19 @@ public:
 	void RemoveObject(AActor* InActor);
 
 	/**
-	 * All the objets that are using this media texture will be used in our mip map calculations.
+	 * All the objects that are using this media texture will be used in our mip map calculations.
 	 *
 	 * @param InMediaTexture Media texture to get objects from.
 	 */
 	void AddObjectsUsingThisMediaTexture(UMediaTexture* InMediaTexture);
 
+	/**
+	 * Remove objects that are using this media texture from our mip map calculations.
+	 *
+	 * @param InMediaTexture Media texture to get objects from.
+	 */
+	void RemoveObjectsUsingThisMediaTexture(UMediaTexture* InMediaTexture);
+
 	/**
 	 * Remove all objects from consideration.
 	 */
diff --git a/Engine/Plugins/Media/ImgMedia/Source/ImgMedia/Private/Loader/ImgMediaLoader.h b/Engine/Plugins/Media/ImgMedia/Source/ImgMedia/Private/Loader/ImgMediaLoader.h
--- a/Engine/Plugins/Media/ImgMedia/Source/ImgMedia/Private/Loader/ImgMediaLoader.h
+++ b/Engine/Plugins/Media/ImgMedia/Source/ImgMedia/Private/Loader/ImgMediaLoader.h
@@ -160,6 +160,16 @@
 	}
 
 	/**
+	 * Get the mipmap info object used by this loader, if any.
+	 *
+	 * @return The mipmap info object.
+	 */
+	const TSharedPtr<FImgMediaMipMapInfo, ESPMode::ThreadSafe>& GetMipMapInfo() const
+	{
+		return MipMapInfo;
+	}
+
+	/**
 	 * Get the width and height of the image sequence.
 	 *
 	 * The dimensions of the image sequence are determined by reading the attributes
diff --git a/Engine/Plugins/Media/ImgMedia/Source/ImgMedia/Private/Player/ImgMediaPlayer.cpp b/Engine/Plugins/Media/ImgMedia/Source/ImgMedia/Private/Player/ImgMediaPlayer.cpp
--- a/Engine/Plugins/Media/ImgMedia/Source/ImgMedia/Private/Player/ImgMediaPlayer.cpp
+++ b/Engine/Plugins/Media/ImgMedia/Source/ImgMedia/Private/Player/ImgMediaPlayer.cpp
@@ -36,6 +36,38 @@ DECLARE_CYCLE_STAT(TEXT("ImgMedia Player TickInput"), STAT_ImgMedia_PlayerTickIn
 
 const FTimespan HackDeltaTimeOffset(1);
 
+namespace {
+	/** Convenience function to process all media textures corresponding the specified player. */
+	void ApplyToPlayerMediaTextures(FImgMediaPlayer* InPlayer, TFunctionRef<void(UMediaTexture*)> TextureCallbackFn)
+	{
+		TArray<UMediaTexture*> PlayerTextures;
+
+		FMediaTextureTracker& TextureTracker = FMediaTextureTracker::Get();
+
+		// Look through all the media textures we know about.
+		for (TWeakObjectPtr<UMediaTexture> TexturePtr : TextureTracker.GetTextures())
+		{
+			UMediaTexture* Texture = TexturePtr.Get();
+			if (Texture != nullptr)
+			{
+				// Does this match the player?
+				UMediaPlayer* MediaPlayer = Texture->GetMediaPlayer();
+				if (MediaPlayer != nullptr)
+				{
+					TSharedPtr<IMediaPlayer, ESPMode::ThreadSafe> Player = MediaPlayer->GetPlayerFacade()->GetPlayer();
+					if (Player.IsValid())
+					{
+						if (Player.Get() == InPlayer)
+						{
+							TextureCallbackFn(Texture);
+						}
+					}
+				}
+			}
+		}
+	}
+}
+
 
 /* FImgMediaPlayer structors
  *****************************************************************************/
@@ -77,6 +109,17 @@ void FImgMediaPlayer::Close()
 		return;
 	}
 
+	{
+		const TSharedPtr<FImgMediaMipMapInfo, ESPMode::ThreadSafe>& MipMapInfo = Loader->GetMipMapInfo();
+		if (MipMapInfo.IsValid())
+		{
+			ApplyToPlayerMediaTextures(this, [&MipMapInfo](UMediaTexture* Texture)
+				{
+					MipMapInfo->RemoveObjectsUsingThisMediaTexture(Texture);
+				});
+		}
+	}
+
 	Scheduler->UnregisterLoader(Loader.ToSharedRef());
 	Loader.Reset();
 
@@ -204,32 +247,10 @@ bool FImgMediaPlayer::Open(const FString& Url, const IMediaOptions* Options)
 			MipMapInfo = StaticCastSharedPtr<FImgMediaMipMapInfo, IMediaOptions::FDataContainer, ESPMode::ThreadSafe>(DataContainer);
 			if (MipMapInfo.IsValid())
 			{
-				// Tell mipmipinfo about our media textures.
-				MipMapInfo->ClearAllObjects();
-
-				FMediaTextureTracker& TextureTracker = FMediaTextureTracker::Get();
-
-				// Look through all the media textures we know about.
-				for (TWeakObjectPtr<UMediaTexture> TexturePtr : TextureTracker.GetTextures())
-				{
-					UMediaTexture* Texture = TexturePtr.Get();
-					if (Texture != nullptr)
+				ApplyToPlayerMediaTextures(this, [&MipMapInfo](UMediaTexture* Texture)
 					{
-						// Does this match the player?
-						UMediaPlayer* MediaPlayer = Texture->GetMediaPlayer();
-						if (MediaPlayer != nullptr)
-						{
-							TSharedPtr<IMediaPlayer, ESPMode::ThreadSafe> Player = MediaPlayer->GetPlayerFacade()->GetPlayer();
-							if (Player.IsValid())
-							{
-								if (Player.Get() == this)
-								{
-									MipMapInfo->AddObjectsUsingThisMediaTexture(Texture);
-								}
-							}
-						}
-					}
-				}
+						MipMapInfo->AddObjectsUsingThisMediaTexture(Texture);
+					});
 			}
 		}
 	}
