ImgMedia ExrReader:
- Added support for multiple tile regions.
- Added ability to read and process only missing tiles instead of re-reading the entire frame.

diff --git a/Engine/Plugins/Media/ImgMedia/Source/ImgMedia/Private/Assets/ImgMediaMipMapInfo.cpp b/Engine/Plugins/Media/ImgMedia/Source/ImgMedia/Private/Assets/ImgMediaMipMapInfo.cpp
--- a/Engine/Plugins/Media/ImgMedia/Source/ImgMedia/Private/Assets/ImgMediaMipMapInfo.cpp
+++ b/Engine/Plugins/Media/ImgMedia/Source/ImgMedia/Private/Assets/ImgMediaMipMapInfo.cpp
@@ -110,7 +110,7 @@ TArray<FIntPoint> FImgMediaTileSelection::GetVisibleCoordinates() const
 	return OutCoordinates;
 }
 
-TArray<FIntRect> FImgMediaTileSelection::GetVisibleRegions() const
+TArray<FIntRect> FImgMediaTileSelection::GetVisibleRegions(const FImgMediaTileSelection* CurrentTileSelection) const
 {
 	/**
 	 * This is a two-pass algorithm to batch tiles into contiguous regions, with a bias for row groupings.
@@ -123,14 +123,21 @@ TArray<FIntRect> FImgMediaTileSelection::GetVisibleRegions() const
 	for (int32 CoordY = 0; CoordY < Dimensions.Y; ++CoordY)
 	{
 		TArray<FIntRect> RegionsPerRow;
-
+		TBitArray<> PreviousVisibleTiles(0, Dimensions.X);
 		for (int32 CoordX = 0; CoordX < Dimensions.X; ++CoordX)
 		{
-			if (Tiles[ToIndex(CoordX, CoordY, Dimensions)])
+			int32 TileIndex = ToIndex(CoordX, CoordY, Dimensions);
+
+			bool bOnlyIncludeMissingTiles = CurrentTileSelection != nullptr;
+
+			// Interpretation: If cached selection doesn't have a tile and CurrentTileSelection (latest) does, then we need to count it as a missing tile.
+			bool bIsThisTileMissing = bOnlyIncludeMissingTiles ? (!Tiles[TileIndex] && CurrentTileSelection->Tiles[TileIndex]) : (Tiles[TileIndex]);
+			if (bIsThisTileMissing)
 			{
 				FIntPoint TileCoord(CoordX, CoordY);
+				PreviousVisibleTiles[CoordX] = true;
+				bool bIsPreviousRowTileVisible = (CoordX > 0) ? PreviousVisibleTiles[CoordX - 1] : false;
 
-				bool bIsPreviousRowTileVisible = (CoordX > 0) ? Tiles[ToIndex(CoordX - 1, CoordY, Dimensions)] : false;
 				if (bIsPreviousRowTileVisible)
 				{
 					RegionsPerRow.Last().Include(TileCoord + 1);
diff --git a/Engine/Plugins/Media/ImgMedia/Source/ImgMedia/Private/Assets/ImgMediaMipMapInfo.h b/Engine/Plugins/Media/ImgMedia/Source/ImgMedia/Private/Assets/ImgMediaMipMapInfo.h
--- a/Engine/Plugins/Media/ImgMedia/Source/ImgMedia/Private/Assets/ImgMediaMipMapInfo.h
+++ b/Engine/Plugins/Media/ImgMedia/Source/ImgMedia/Private/Assets/ImgMediaMipMapInfo.h
@@ -112,11 +112,12 @@ struct FImgMediaTileSelection
 	TArray<FIntPoint> GetVisibleCoordinates() const;
 
 	/**
-	 * Returns a calculated list of contiguous visible tile regions.
+	 * Returns a calculated list of contiguous visible tile regions. Only provides regions for the missing tiles if
+	 * CurrentTileSelection is specified.
 	 *
-	 * @return Visible tile regions array.
+	 * @return Visible tile regions array. 
 	 */
-	TArray<FIntRect> GetVisibleRegions() const;
+	TArray<FIntRect> GetVisibleRegions(const FImgMediaTileSelection* CurrentTileSelection = nullptr) const;
 
 	/**
 	 * Return the rectangular region bounding the visible tiles.
diff --git a/Engine/Plugins/Media/ImgMedia/Source/ImgMedia/Private/Readers/ExrImgMediaReader.cpp b/Engine/Plugins/Media/ImgMedia/Source/ImgMedia/Private/Readers/ExrImgMediaReader.cpp
--- a/Engine/Plugins/Media/ImgMedia/Source/ImgMedia/Private/Readers/ExrImgMediaReader.cpp
+++ b/Engine/Plugins/Media/ImgMedia/Source/ImgMedia/Private/Readers/ExrImgMediaReader.cpp
@@ -56,11 +56,11 @@ bool FExrImgMediaReader::GetFrameInfo(const FString& ImagePath, FImgMediaFrameIn
 }
 
 FExrImgMediaReader::EReadResult FExrImgMediaReader::ReadTilesCustom
-(uint16* Buffer
+	( uint16* Buffer
 	, int64 BufferSize
 	, const FString& ImagePath
 	, int32 FrameId
-	, const FIntRect& TileRegion
+	, const TArray<FIntRect>& TileRegions
 	, TSharedPtr<FSampleConverterParameters> ConverterParams
 	, const int32 CurrentMipLevel)
 {
@@ -100,56 +100,62 @@ FExrImgMediaReader::EReadResult FExrImgMediaReader::ReadTilesCustom
 	}
 	{
 		TRACE_CPUPROFILER_EVENT_SCOPE_TEXT(*FString::Printf(TEXT("FExrImgMediaReader_ReadTilesCustom_ReadTiles")));
-
-		for (int32 TileRow = TileRegion.Min.Y; TileRow < TileRegion.Max.Y; TileRow++)
+		for (const FIntRect& TileRegion : TileRegions)
 		{
-			// Check to see if the frame was canceled.
+			for (int32 TileRow = TileRegion.Min.Y; TileRow < TileRegion.Max.Y; TileRow++)
 			{
-				FScopeLock RegionScopeLock(&CanceledFramesCriticalSection);
-				if (CanceledFrames.Remove(FrameId) > 0)
+				// Check to see if the frame was canceled.
 				{
-					UE_LOG(LogImgMedia, Warning, TEXT("Reader %p: Canceling Frame %i At tile row # %i"), this, FrameId, TileRow);
-					bResult = Cancelled;
-					break;
+					FScopeLock RegionScopeLock(&CanceledFramesCriticalSection);
+					if (CanceledFrames.Remove(FrameId) > 0)
+					{
+						UE_LOG(LogImgMedia, Warning, TEXT("Reader %p: Canceling Frame %i At tile row # %i"), this, FrameId, TileRow);
+						bResult = Cancelled;
+						break;
+					}
 				}
-			}
 
-			const uint16 Padding = ConverterParams->bCustomExr ? 0 : FExrReader::TILE_PADDING;
-			const int64 TileByteStride = ConverterParams->PixelSize * ConverterParams->TileDimWithBorders.X * ConverterParams->TileDimWithBorders.Y + Padding;
-			const int StartTileIndex = TileRow * DimensionInTiles.X + TileRegion.Min.X;
+				const uint16 Padding = ConverterParams->bCustomExr ? 0 : FExrReader::TILE_PADDING;
+				const int64 TileByteStride = ConverterParams->PixelSize * ConverterParams->TileDimWithBorders.X * ConverterParams->TileDimWithBorders.Y + Padding;
+				const int StartTileIndex = TileRow * DimensionInTiles.X + TileRegion.Min.X;
 
-			bool bLastTile = TileRow + 1 == DimensionInTiles.Y && TileRegion.Max.X == DimensionInTiles.X;
-			const int EndTileIndex = TileRow * DimensionInTiles.X + TileRegion.Max.X;
+				bool bLastTile = TileRow + 1 == DimensionInTiles.Y && TileRegion.Max.X == DimensionInTiles.X;
+				const int EndTileIndex = TileRow * DimensionInTiles.X + TileRegion.Max.X;
 
-			int MipLevel = ConverterParams->bMipsInSeparateFiles ? 0 : CurrentMipLevel;
-			ChunkReader.SeekTileWithinFile(StartTileIndex, MipLevel, CurrentBufferPos);
-			int64 ByteOffsetStart = 0;
-			int64 ByteOffsetEnd = 0;
+				int MipLevel = ConverterParams->bMipsInSeparateFiles ? 0 : CurrentMipLevel;
+				ChunkReader.SeekTileWithinFile(StartTileIndex, MipLevel, CurrentBufferPos);
+				int64 ByteOffsetStart = 0;
+				int64 ByteOffsetEnd = 0;
 
-			if (!ChunkReader.GetByteOffsetForTile(StartTileIndex, MipLevel, ByteOffsetStart)
-				|| !ChunkReader.GetByteOffsetForTile(EndTileIndex, MipLevel, ByteOffsetEnd))
-			{
-				bResult = Fail;
-				break;
-			}
+				if (!ChunkReader.GetByteOffsetForTile(StartTileIndex, MipLevel, ByteOffsetStart)
+					|| !ChunkReader.GetByteOffsetForTile(EndTileIndex, MipLevel, ByteOffsetEnd))
+				{
+					bResult = Fail;
+					break;
+				}
 
-			// If this is the last tile, make sure chunk reader reads till the end of the buffer we write into.
-			int64 ByteChunkToRead = bLastTile
-				? FMath::Min(ByteOffsetEnd - ByteOffsetStart, BufferSize - CurrentBufferPos)
-				: (ByteOffsetEnd - ByteOffsetStart);
+				// If this is the last tile, make sure chunk reader reads till the end of the buffer we write into.
+				int64 ByteChunkToRead = bLastTile
+					? FMath::Min(ByteOffsetEnd - ByteOffsetStart, BufferSize - CurrentBufferPos)
+					: (ByteOffsetEnd - ByteOffsetStart);
 
-			if (!ChunkReader.ReadExrImageChunk(reinterpret_cast<char*>(Buffer) + CurrentBufferPos, ByteChunkToRead))
-			{
-				bResult = Fail;
-				break;
+				if (!ChunkReader.ReadExrImageChunk(reinterpret_cast<char*>(Buffer) + CurrentBufferPos, ByteChunkToRead))
+				{
+					bResult = Fail;
+					break;
+				}
 			}
 		}
 	}
-	if (!ChunkReader.CloseExrFile())
 	{
-		return Fail;
+		TRACE_CPUPROFILER_EVENT_SCOPE_TEXT(*FString::Printf(TEXT("FExrImgMediaReaderGpu_CloseFile %d"), CurrentMipLevel));
+
+		if (!ChunkReader.CloseExrFile())
+		{
+			return Fail;
+		}
+		return bResult;
 	}
-	return bResult;
 #else
 	return Fail;
 #endif
@@ -249,7 +255,7 @@ bool FExrImgMediaReader::ReadFrame(int32 FrameId, const TMap<int32, FImgMediaTil
 				if (OutFrame->Info.FormatName == TEXT("EXR CUSTOM"))
 				{
 #if defined(PLATFORM_WINDOWS) && PLATFORM_WINDOWS
-					FIntRect TileRegion = CurrentTileSelection.GetVisibleRegion();
+					TArray<FIntRect> TileRegions = CurrentTileSelection.GetVisibleRegions();
 					int32 PixelSize = sizeof(uint16) * OutFrame->Info.NumChannels;
 					TSharedPtr<FSampleConverterParameters> ConverterParams = MakeShared<FSampleConverterParameters>();
 					ConverterParams->FrameInfo = OutFrame->Info;
@@ -258,7 +264,7 @@ bool FExrImgMediaReader::ReadFrame(int32 FrameId, const TMap<int32, FImgMediaTil
 					ConverterParams->NumMipLevels = Loader->GetNumMipLevels();
 					ConverterParams->bCustomExr = OutFrame->Info.FormatName == TEXT("EXR CUSTOM");
 
-					EReadResult ReadResult = ReadTilesCustom((uint16*)MipDataPtr, GetMipBufferTotalSize(Dim / MipLevelDiv), Image, FrameId, TileRegion, ConverterParams, CurrentMipLevel);
+					EReadResult ReadResult = ReadTilesCustom((uint16*)MipDataPtr, GetMipBufferTotalSize(Dim / MipLevelDiv), Image, FrameId, TileRegions, ConverterParams, CurrentMipLevel);
 					if (ReadResult != Fail)
 					{
 						OutFrame->MipTilesPresent.Emplace(CurrentMipLevel, CurrentTileSelection);
diff --git a/Engine/Plugins/Media/ImgMedia/Source/ImgMedia/Private/Readers/ExrImgMediaReader.h b/Engine/Plugins/Media/ImgMedia/Source/ImgMedia/Private/Readers/ExrImgMediaReader.h
--- a/Engine/Plugins/Media/ImgMedia/Source/ImgMedia/Private/Readers/ExrImgMediaReader.h
+++ b/Engine/Plugins/Media/ImgMedia/Source/ImgMedia/Private/Readers/ExrImgMediaReader.h
@@ -63,8 +63,8 @@ protected:
 		/** Dimension of the tile including the overscan borders. */
 		FIntPoint TileDimWithBorders;
 
-		/** Used for rendering tiles in bulk regions. */
-		TMap<int32, FIntRect> Viewports;
+		/** Used for rendering tiles in bulk regions per mip level. */
+		TMap<int32, TArray<FIntRect>> Viewports;
 
 #if defined(PLATFORM_WINDOWS) && PLATFORM_WINDOWS
 		/** Contain information about individual tiles. Used to convert buffer data into a 2D Texture. */
@@ -97,7 +97,7 @@ protected:
 	 * Reads tiles from exr files in tile rows based on Tile region. If frame is pending for cancelation
 	 * stops reading tiles at the current tile row.
 	*/
-	EReadResult ReadTilesCustom(uint16* Buffer, int64 BufferSize, const FString& ImagePath, int32 FrameId, const FIntRect& TileRegion, TSharedPtr<FSampleConverterParameters> ConverterParams, const int32 CurrentMipLevel);
+	EReadResult ReadTilesCustom(uint16* Buffer, int64 BufferSize, const FString& ImagePath, int32 FrameId, const TArray<FIntRect>& TileRegions, TSharedPtr<FSampleConverterParameters> ConverterParams, const int32 CurrentMipLevel);
 
 	/**
 	 * Sets parameters of our custom format images.
diff --git a/Engine/Plugins/Media/ImgMedia/Source/ImgMedia/Private/Readers/ExrImgMediaReaderGpu.cpp b/Engine/Plugins/Media/ImgMedia/Source/ImgMedia/Private/Readers/ExrImgMediaReaderGpu.cpp
--- a/Engine/Plugins/Media/ImgMedia/Source/ImgMedia/Private/Readers/ExrImgMediaReaderGpu.cpp
+++ b/Engine/Plugins/Media/ImgMedia/Source/ImgMedia/Private/Readers/ExrImgMediaReaderGpu.cpp
@@ -149,16 +149,21 @@
 		return false;
 	}
 
-	TArray<FStructuredBufferPoolItemSharedPtr> BufferDataArray;
-	BufferDataArray.AddDefaulted(FImgMediaLoader::MAX_MIPMAP_LEVELS);
-
 	ConverterParams->FrameInfo = OutFrame->Info;
 	ConverterParams->PixelSize = sizeof(uint16) * ConverterParams->FrameInfo.NumChannels;
 	ConverterParams->TileDimWithBorders = TileDim + OutFrame->Info.TileBorder * 2;
 	ConverterParams->NumMipLevels = Loader->GetNumMipLevels();
 	ConverterParams->bCustomExr = OutFrame->Info.FormatName == TEXT("EXR CUSTOM");
 	ConverterParams->bMipsInSeparateFiles = Loader->MipsInSeparateFiles();
+
+	FExrMediaTextureSampleConverter* SampleConverter;
+	if (!OutFrame->SampleConverter.IsValid())
 	{
+		OutFrame->SampleConverter = MakeShared<FExrMediaTextureSampleConverter>();
+	}
+	SampleConverter = static_cast<FExrMediaTextureSampleConverter*>(OutFrame->SampleConverter.Get());
+
+	{
 		// Loop over all mips.
 		FIntPoint CurrentMipDim = ConverterParams->FullResolution;
 
@@ -178,75 +183,112 @@
 				ReadThisMip = !CachedSelection->Contains(CurrentTileSelection);
 			}
 			
+			if (!ReadThisMip)
+			{
+				continue;
+			}
+
 			// Next mip level.
 			int MipLevelDiv = 1 << CurrentMipLevel;
 			CurrentMipDim = ConverterParams->FullResolution / MipLevelDiv;
 
-			if (ReadThisMip)
+
+			const SIZE_T BufferSize = GetBufferSize(CurrentMipDim, ConverterParams->FrameInfo.NumChannels, bHasTiles, OutFrame->Info.NumTiles / MipLevelDiv, ConverterParams->bCustomExr);
+			FStructuredBufferPoolItemSharedPtr BufferData = SampleConverter->GetMipLevelBuffer(CurrentMipLevel);
+
+			if (!BufferData.IsValid())
 			{
-				FIntRect TileRegion = CurrentTileSelection.GetVisibleRegion();
-				FIntRect& Viewport = ConverterParams->Viewports.Add(CurrentMipLevel);
-				Viewport.Min = FIntPoint(ConverterParams->TileDimWithBorders.X * TileRegion.Min.X, ConverterParams->TileDimWithBorders.Y * TileRegion.Min.Y);
-				Viewport.Max = FIntPoint(ConverterParams->TileDimWithBorders.X * TileRegion.Max.X, ConverterParams->TileDimWithBorders.Y * TileRegion.Max.Y);
-				Viewport.Clip(FIntRect(FIntPoint::ZeroValue, CurrentMipDim));
-
-				const SIZE_T BufferSize = GetBufferSize(CurrentMipDim, ConverterParams->FrameInfo.NumChannels, bHasTiles, OutFrame->Info.NumTiles / MipLevelDiv, ConverterParams->bCustomExr);
-				FStructuredBufferPoolItemSharedPtr& BufferData = BufferDataArray[CurrentMipLevel];
 				BufferData = AllocateGpuBufferFromPool(BufferSize);
-				uint16* MipDataPtr = static_cast<uint16*>(BufferData->MappedBuffer);
+				SampleConverter->SetMipLevelBuffer(CurrentMipLevel, BufferData);
+			}
 
-				// Get highest resolution mip level path.
-				FString ImagePath = Loader->GetImagePath(FrameId, ConverterParams->bMipsInSeparateFiles ? CurrentMipLevel : 0);
+			uint16* MipDataPtr = static_cast<uint16*>(BufferData->MappedBuffer);
+
+			// Get highest resolution mip level path.
+			FString ImagePath = Loader->GetImagePath(FrameId, ConverterParams->bMipsInSeparateFiles ? CurrentMipLevel : 0);
 				
-				EReadResult ReadResult = Fail;
+			EReadResult ReadResult = Fail;
 
-				FOpenExrHeaderReader InputTileFile(ImagePath);
-				if (InputTileFile.HasInputFile())
+			if (FPaths::FileExists(ImagePath))
+			{
+				// read frame data
+				if (bHasTiles || ConverterParams->bCustomExr)
 				{
-					// read frame data
-					if (bHasTiles || ConverterParams->bCustomExr)
+					TArray<FIntRect> TileRegionsToRead;
+					TArray<FIntRect> TileRegionsToRender;
 					{
-						ReadResult = ReadTilesCustom(MipDataPtr, BufferSize, ImagePath, FrameId, TileRegion, ConverterParams, CurrentMipLevel);
-					}
-					else
-					{
-						ReadResult = ReadInChunks(MipDataPtr, ImagePath, FrameId, CurrentMipDim, BufferSize);
+						TRACE_CPUPROFILER_EVENT_SCOPE_TEXT(*FString::Printf(TEXT("CalculateRegions %d"), CurrentMipLevel));
+
+						if (const FImgMediaTileSelection* CachedSelection = OutFrame->MipTilesPresent.Find(CurrentMipLevel))
+						{
+							TileRegionsToRead = CachedSelection->GetVisibleRegions(&CurrentTileSelection);
+							TileRegionsToRender = CurrentTileSelection.GetVisibleRegions();
+						}
+						else
+						{
+							TileRegionsToRead = TileRegionsToRender = CurrentTileSelection.GetVisibleRegions();
+						}
 					}
 
-					if (ReadResult != Fail)
+					TArray<FIntRect>& Viewports = ConverterParams->Viewports.Add(CurrentMipLevel);
+					for (const FIntRect& TileRegion : TileRegionsToRender)
 					{
-						OutFrame->MipTilesPresent.Emplace(CurrentMipLevel, CurrentTileSelection);
+						FIntRect Viewport;
+						Viewport.Min = FIntPoint(ConverterParams->TileDimWithBorders.X * TileRegion.Min.X, ConverterParams->TileDimWithBorders.Y * TileRegion.Min.Y);
+						Viewport.Max = FIntPoint(ConverterParams->TileDimWithBorders.X * TileRegion.Max.X, ConverterParams->TileDimWithBorders.Y * TileRegion.Max.Y);
+						Viewport.Clip(FIntRect(FIntPoint::ZeroValue, CurrentMipDim));
+						Viewports.Add(MoveTemp(Viewport));
 					}
+
+					ReadResult = ReadTilesCustom(MipDataPtr, BufferSize, ImagePath, FrameId, TileRegionsToRead, ConverterParams, CurrentMipLevel);
 				}
 				else
 				{
-					UE_LOG(LogImgMedia, Error, TEXT("Could not load %s"), *ImagePath);
-					return false;
+					TArray<FIntRect>& Viewports = ConverterParams->Viewports.Add(CurrentMipLevel);
+					FIntRect Viewport;
+						
+					Viewport.Min = FIntPoint(0,0);
+					Viewport.Max = CurrentMipDim;
+					Viewports.Add(MoveTemp(Viewport));
+
+					ReadResult = ReadInChunks(MipDataPtr, ImagePath, FrameId, CurrentMipDim, BufferSize);
 				}
-				if (ReadResult == Fail)
+
+				if (ReadResult != Fail)
 				{
-					// Check if we have a compressed file.
-					FImgMediaFrameInfo Info;
-					if (GetInfo(ImagePath, Info))
+					OutFrame->MipTilesPresent.Emplace(CurrentMipLevel, CurrentTileSelection);
+				}
+			}
+			else
+			{
+				UE_LOG(LogImgMedia, Error, TEXT("Could not load %s"), *ImagePath);
+				return false;
+			}
+			if (ReadResult == Fail)
+			{
+				// Check if we have a compressed file.
+				FImgMediaFrameInfo Info;
+				if (GetInfo(ImagePath, Info))
+				{
+					if (Info.CompressionName != "Uncompressed")
 					{
-						if (Info.CompressionName != "Uncompressed")
-						{
-							UE_LOG(LogImgMedia, Error, TEXT("GPU Reader cannot read compressed file %s."), *ImagePath);
-							UE_LOG(LogImgMedia, Error, TEXT("Compressed and uncompressed files should not be mixed in a single sequence."));
-						}
+						UE_LOG(LogImgMedia, Error, TEXT("GPU Reader cannot read compressed file %s."), *ImagePath);
+						UE_LOG(LogImgMedia, Error, TEXT("Compressed and uncompressed files should not be mixed in a single sequence."));
 					}
+				}
 
-					// Fall back to CPU.
-					bFallBackToCPU = true;
-					return FExrImgMediaReader::ReadFrame(FrameId, InMipTiles, OutFrame);
-				}
+				// Fall back to CPU.
+				bFallBackToCPU = true;
+				return FExrImgMediaReader::ReadFrame(FrameId, InMipTiles, OutFrame);
 			}
 		}
 	}
+
 	OutFrame->Format = ConverterParams->FrameInfo.NumChannels <= 3 ? EMediaTextureSampleFormat::FloatRGB : EMediaTextureSampleFormat::FloatRGBA;
 	OutFrame->Stride = ConverterParams->FullResolution.X * ConverterParams->PixelSize;
-	OutFrame->SampleConverter = CreateSampleConverter(MoveTemp(BufferDataArray), ConverterParams);
 
+	CreateSampleConverterCallback(SampleConverter, ConverterParams);
+
 	UE_LOG(LogImgMedia, Verbose, TEXT("Reader %p: Read Pixels Complete. %i"), this, FrameId);
 	return true;
 }
@@ -359,25 +401,24 @@
 
 }
 
-TSharedPtr<IMediaTextureSampleConverter, ESPMode::ThreadSafe> FExrImgMediaReaderGpu::CreateSampleConverter
-	( TArray<FStructuredBufferPoolItemSharedPtr>&& BufferDataArray
+void FExrImgMediaReaderGpu::CreateSampleConverterCallback
+	( FExrMediaTextureSampleConverter* SampleConverter
 	, TSharedPtr<FSampleConverterParameters> ConverterParams)
 {
-	auto RenderThreadSwizzler = [BufferDataArray, ConverterParams] (FRHICommandListImmediate& RHICmdList, FTexture2DRHIRef RenderTargetTextureRHI)->bool
+	auto RenderThreadSwizzler = [ConverterParams] (FRHICommandListImmediate& RHICmdList, FTexture2DRHIRef RenderTargetTextureRHI, TMap<int32, FStructuredBufferPoolItemSharedPtr>& MipBuffers)->bool
 	{
-
 		SCOPED_DRAW_EVENT(RHICmdList, FExrImgMediaReaderGpu_Convert);
 		SCOPED_GPU_STAT(RHICmdList, ExrImgMediaReaderGpu);
 
 		FIntPoint Dim = ConverterParams->FullResolution;
-		for (int32 MipLevel = 0; MipLevel < ConverterParams->NumMipLevels; ++MipLevel)
+		for (TPair<int32, TArray<FIntRect>> MipLevelViewports : ConverterParams->Viewports)
 		{
+			int32 MipLevel = MipLevelViewports.Key;
 			SCOPED_GPU_STAT(RHICmdList, ExrImgMediaReaderGpu_MipRender);
-			int MipLevelDiv = 1 << MipLevel;
+			int MipLevelDiv = 1 << MipLevelViewports.Key;
 			Dim = ConverterParams->FullResolution / MipLevelDiv;
 
-
-			FStructuredBufferPoolItemSharedPtr BufferData = BufferDataArray[MipLevel];
+			FStructuredBufferPoolItemSharedPtr BufferData = MipBuffers[MipLevel];
 			if (BufferData.IsValid())
 			{
 				if (!BufferData->BufferRef->IsValid())
@@ -385,8 +426,6 @@
 					continue;
 				}
 
-				FIntRect Viewport = ConverterParams->Viewports[MipLevel];
-
 				// This flag will indicate that we should wait for poll to complete.
 				BufferData->bWillBeSignaled = true;
 
@@ -434,10 +473,15 @@
 				TShaderMapRef<FExrSwizzlePS> SwizzleShaderPS(ShaderMap, PermutationVector);
 
 				FScreenPassPipelineState PipelineState(SwizzleShaderVS, SwizzleShaderPS, TStaticBlendState<>::GetRHI(), TStaticDepthStencilState<false, CF_Always>::GetRHI());
-				DrawScreenPass(RHICmdList, Dim, Viewport, PipelineState, [&](FRHICommandListImmediate& RHICmdList)
-					{
-						SetShaderParameters(RHICmdList, SwizzleShaderPS, SwizzleShaderPS.GetPixelShader(), Parameters);
-					});
+				
+				// If there are tiles determines if we should deliver tiles one by one or in a bulk.
+				for (const FIntRect& Viewport : MipLevelViewports.Value)
+				{
+					DrawScreenPass(RHICmdList, Dim, Viewport, PipelineState, [&](FRHICommandListImmediate& RHICmdList)
+						{
+							SetShaderParameters(RHICmdList, SwizzleShaderPS, SwizzleShaderPS.GetPixelShader(), Parameters);
+						});
+				}
 
 				// Resolve render target.
 				RHICmdList.EndRenderPass();
@@ -453,10 +497,8 @@
 		return false;
 	};
 
-	FExrMediaTextureSampleConverter* SampleConverter = new FExrMediaTextureSampleConverter();
-	SampleConverter->ConvertExrBufferCallback = FExrConvertBufferCallback::CreateLambda(RenderThreadSwizzler);
-
-	return MakeShareable(SampleConverter);
+	// Stacks up converters for each tile region.
+	SampleConverter->AddCallback(FExrConvertBufferCallback::CreateLambda(RenderThreadSwizzler));
 }
 
 FStructuredBufferPoolItemSharedPtr FExrImgMediaReaderGpu::AllocateGpuBufferFromPool(uint32 AllocSize, bool bWait)
@@ -496,7 +538,7 @@
 				SCOPED_DRAW_EVENT(RHICmdList, FExrImgMediaReaderGpu_AllocateBuffer);
 				FRHIResourceCreateInfo CreateInfo(TEXT("FExrImgMediaReaderGpu"));
 				AllocatedBuffer->BufferRef = RHICreateStructuredBuffer(sizeof(uint16) * 2., AllocSize, BUF_ShaderResource | BUF_Dynamic | BUF_FastVRAM, CreateInfo);
-				AllocatedBuffer->MappedBuffer = static_cast<uint16*>(RHILockBuffer(AllocatedBuffer->BufferRef, 0, AllocSize, RLM_WriteOnly));
+				AllocatedBuffer->MappedBuffer = RHILockBuffer(AllocatedBuffer->BufferRef, 0, AllocSize, RLM_WriteOnly);
 				AllocatedBuffer->Fence = RHICreateGPUFence(TEXT("BufferNoLongerInUseFence"));
 				if (bWait)
 				{
@@ -579,7 +621,9 @@
 
 bool FExrMediaTextureSampleConverter::Convert(FTexture2DRHIRef& InDstTexture, const FConversionHints& Hints)
 {
-	return ConvertExrBufferCallback.Execute(FRHICommandListExecutor::GetImmediateCommandList(), InDstTexture);
+	FScopeLock ScopeLock(&ConverterCallbacksCriticalSection);
+	bool bExecutionSuccessful = ConvertExrBufferCallback.Execute(FRHICommandListExecutor::GetImmediateCommandList(), InDstTexture, MipBuffers);
+	return bExecutionSuccessful;
 }
 
 #endif //IMGMEDIA_EXR_SUPPORTED_PLATFORM && PLATFORM_WINDOWS
diff --git a/Engine/Plugins/Media/ImgMedia/Source/ImgMedia/Private/Readers/ExrImgMediaReaderGpu.h b/Engine/Plugins/Media/ImgMedia/Source/ImgMedia/Private/Readers/ExrImgMediaReaderGpu.h
--- a/Engine/Plugins/Media/ImgMedia/Source/ImgMedia/Private/Readers/ExrImgMediaReaderGpu.h
+++ b/Engine/Plugins/Media/ImgMedia/Source/ImgMedia/Private/Readers/ExrImgMediaReaderGpu.h
@@ -10,6 +10,7 @@
 #include "IMediaTextureSampleConverter.h"
 
 class FExrImgMediaReaderGpu;
+class FExrMediaTextureSampleConverter;
 
 struct FStructuredBufferPoolItem
 {
@@ -89,8 +90,8 @@ protected:
 	/**
 	* Creates Sample converter to be used by Media Texture Resource.
 	*/
-	static TSharedPtr<IMediaTextureSampleConverter, ESPMode::ThreadSafe> CreateSampleConverter
-		(TArray<FStructuredBufferPoolItemSharedPtr>&& BufferDataArray, TSharedPtr<FSampleConverterParameters> ConverterParams);
+	static void CreateSampleConverterCallback
+		(FExrMediaTextureSampleConverter* SampleConverter, TSharedPtr<FSampleConverterParameters> ConverterParams);
 
 public:
 
@@ -127,7 +128,7 @@ private:
 	bool bFallBackToCPU;
 };
 
-FUNC_DECLARE_DELEGATE(FExrConvertBufferCallback, bool, FRHICommandListImmediate& /*RHICmdList*/, FTexture2DRHIRef /*RenderTargetTextureRHI*/)
+FUNC_DECLARE_DELEGATE(FExrConvertBufferCallback, bool, FRHICommandListImmediate& /*RHICmdList*/, FTexture2DRHIRef /*RenderTargetTextureRHI*/, TMap<int32, FStructuredBufferPoolItemSharedPtr>& /*MipBuffers*/ )
 
 class FExrMediaTextureSampleConverter: public IMediaTextureSampleConverter
 {
@@ -135,9 +136,36 @@ class FExrMediaTextureSampleConverter: public IMediaTextureSampleConverter
 public:
 	virtual bool Convert(FTexture2DRHIRef& InDstTexture, const FConversionHints& Hints) override;
 	virtual ~FExrMediaTextureSampleConverter() {};
+	
+	void AddCallback(FExrConvertBufferCallback&& Callback) 
+	{
+		FScopeLock ScopeLock(&ConverterCallbacksCriticalSection);
+		ConvertExrBufferCallback = Callback;
+	};
+
+	FStructuredBufferPoolItemSharedPtr GetMipLevelBuffer(int32 RequestedMipLevel)
+	{
+		FScopeLock ScopeLock(&ConverterCallbacksCriticalSection);
+		if (MipBuffers.Contains(RequestedMipLevel))
+		{
+			return *MipBuffers.Find(RequestedMipLevel);
+		}
+		return nullptr;
+	}
+
+	void SetMipLevelBuffer(int32 RequestedMipLevel, FStructuredBufferPoolItemSharedPtr Buffer)
+	{
+		FScopeLock ScopeLock(&ConverterCallbacksCriticalSection);
+		check(!MipBuffers.Contains(RequestedMipLevel));
+		MipBuffers.Add(RequestedMipLevel, Buffer);
+	}
 
-public:
+private:
+	FCriticalSection ConverterCallbacksCriticalSection;
 	FExrConvertBufferCallback ConvertExrBufferCallback;
+
+	/** An array of structured buffers that are big enough to fully contain corresponding mip levels. */
+	TMap<int32,FStructuredBufferPoolItemSharedPtr> MipBuffers;
 };
 
 #endif //defined(PLATFORM_WINDOWS) && PLATFORM_WINDOWS
