ImgMedia: Fixed a bug where tile locations for custom exr were not calculated correctly.

diff --git a/Engine/Plugins/Media/ImgMedia/Source/ExrReaderGpu/Private/ExrReaderGpu.cpp b/Engine/Plugins/Media/ImgMedia/Source/ExrReaderGpu/Private/ExrReaderGpu.cpp
--- a/Engine/Plugins/Media/ImgMedia/Source/ExrReaderGpu/Private/ExrReaderGpu.cpp
+++ b/Engine/Plugins/Media/ImgMedia/Source/ExrReaderGpu/Private/ExrReaderGpu.cpp
@@ -264,7 +264,6 @@ void FExrReader::CalculateTileOffsets
 	, TArray<TArray<int64>>& OutCustomOffsets
 	, TArray<TArray<FTileDesc>>& OutPartialTileInfo
 	, const FIntPoint& FullTextureResolution
-	, const FIntPoint& FullResolutionInTiles
 	, const FIntPoint& TileDimWithBorders
 	, int32 NumMipLevels
 	, int64 PixelSize
@@ -280,15 +279,20 @@ void FExrReader::CalculateTileOffsets
 		{
 			int32 MipDiv = 1 << MipLevel;
 
+			// Resolution of the texture in pixels for this mip level.
+			FIntPoint MipResolution = FullTextureResolution / MipDiv;
+
 			// Dimension of the texture in tiles including partial tiles
+			float TileResFractionX = float(MipResolution.X) / TileDimWithBorders.X;
+			float TileResFractionY = float(MipResolution.Y) / TileDimWithBorders.Y;
 			FIntPoint DimensionInTiles_PartialTiles
-				( FMath::CeilToInt(float(FullResolutionInTiles.X) / MipDiv)
-				, FMath::CeilToInt(float(FullResolutionInTiles.Y) / MipDiv));
+				( FMath::CeilToInt(TileResFractionX)
+				, FMath::CeilToInt(TileResFractionY));
 
 			// Dimension of the texture in tiles excluding partial tiles.
 			FIntPoint DimensionInTiles_CompleteTiles
-				( FMath::FloorToInt(float(FullResolutionInTiles.X) / MipDiv)
-				, FMath::FloorToInt(float(FullResolutionInTiles.Y) / MipDiv));
+				( FMath::FloorToInt(TileResFractionX)
+				, FMath::FloorToInt(TileResFractionY));
 
 			// Total number of tiles for this mip level.
 			int32 NumActualTiles = DimensionInTiles_PartialTiles.X * DimensionInTiles_PartialTiles.Y;
@@ -313,8 +317,6 @@ void FExrReader::CalculateTileOffsets
 				continue;
 			}
 
-			// Resolution of the texture in pixels for this mip level.
-			FIntPoint MipResolution = FullTextureResolution / MipDiv;
 
 			// Resolution of the partial tile in the bottom right corner.
 			FIntPoint PartialTileResolution = FIntPoint(MipResolution.X % TileDimWithBorders.X, MipResolution.Y % TileDimWithBorders.Y);
diff --git a/Engine/Plugins/Media/ImgMedia/Source/ExrReaderGpu/Public/ExrReaderGpu.h b/Engine/Plugins/Media/ImgMedia/Source/ExrReaderGpu/Public/ExrReaderGpu.h
--- a/Engine/Plugins/Media/ImgMedia/Source/ExrReaderGpu/Public/ExrReaderGpu.h
+++ b/Engine/Plugins/Media/ImgMedia/Source/ExrReaderGpu/Public/ExrReaderGpu.h
@@ -83,7 +83,6 @@ public:
 		, TArray<TArray<int64>>& OutCustomOffsets
 		, TArray<TArray<FTileDesc>>& OutPartialTileInfo
 		, const FIntPoint& FullTextureResolution
-		, const FIntPoint& FullResolutionInTiles
 		, const FIntPoint& TileDimWithBorders
 		, int32 NumMipLevels
 		, int64 PixelSize
diff --git a/Engine/Plugins/Media/ImgMedia/Source/ImgMedia/Private/Readers/ExrImgMediaReader.cpp b/Engine/Plugins/Media/ImgMedia/Source/ImgMedia/Private/Readers/ExrImgMediaReader.cpp
--- a/Engine/Plugins/Media/ImgMedia/Source/ImgMedia/Private/Readers/ExrImgMediaReader.cpp
+++ b/Engine/Plugins/Media/ImgMedia/Source/ImgMedia/Private/Readers/ExrImgMediaReader.cpp
@@ -71,21 +71,21 @@ FExrImgMediaReader::EReadResult FExrImgMediaReader::ReadTilesCustom
 	FExrReader ChunkReader;
 	int MipLevelDiv = 1 << CurrentMipLevel;
 
+	FIntPoint MipResolution = ConverterParams->FullResolution / MipLevelDiv;
+
 	FIntPoint DimensionInTiles
-	(FMath::CeilToInt(float(ConverterParams->FrameInfo.NumTiles.X) / MipLevelDiv)
-		, FMath::CeilToInt(float(ConverterParams->FrameInfo.NumTiles.Y) / MipLevelDiv));
+		( FMath::CeilToInt(float(MipResolution.X) / ConverterParams->TileDimWithBorders.X)
+		, FMath::CeilToInt(float(MipResolution.Y) / ConverterParams->TileDimWithBorders.Y));
 
 	TArray<int32> NumTilesPerLevel;
 	TArray<TArray<int64>> CustomOffsets;
 	int32 NumMipLevels = ConverterParams->bMipsInSeparateFiles ? 1 : ConverterParams->NumMipLevels;
-	bool bHasPartialTiles = false;
 
 	FExrReader::CalculateTileOffsets(
 		NumTilesPerLevel,
 		CustomOffsets,
 		ConverterParams->TileInfoPerMipLevel,
 		ConverterParams->FullResolution,
-		ConverterParams->FrameInfo.NumTiles,
 		ConverterParams->TileDimWithBorders,
 		NumMipLevels,
 		ConverterParams->PixelSize,
@@ -404,7 +404,9 @@ bool FExrImgMediaReader::GetInfo(const FString& FilePath, FImgMediaFrameInfo& Ou
 
 	if (OutInfo.bHasTiles)
 	{
-		OutInfo.NumTiles = FIntPoint(OutInfo.Dim.X / (OutInfo.TileDimensions.X + OutInfo.TileBorder * 2), OutInfo.Dim.Y / (OutInfo.TileDimensions.Y + OutInfo.TileBorder * 2));
+		OutInfo.NumTiles =
+		{FMath::CeilToInt(float(OutInfo.Dim.X) / (OutInfo.TileDimensions.X + OutInfo.TileBorder * 2))
+		, FMath::CeilToInt(float(OutInfo.Dim.Y) / (OutInfo.TileDimensions.Y + OutInfo.TileBorder * 2))};
 		if (HeaderReader.ContainsMips())
 		{
 			OutInfo.NumMipLevels = HeaderReader.CalculateNumMipLevels(OutInfo.NumTiles);
diff --git a/Engine/Plugins/Media/ImgMedia/Source/ImgMedia/Private/Readers/ExrImgMediaReaderGpu.cpp b/Engine/Plugins/Media/ImgMedia/Source/ImgMedia/Private/Readers/ExrImgMediaReaderGpu.cpp
--- a/Engine/Plugins/Media/ImgMedia/Source/ImgMedia/Private/Readers/ExrImgMediaReaderGpu.cpp
+++ b/Engine/Plugins/Media/ImgMedia/Source/ImgMedia/Private/Readers/ExrImgMediaReaderGpu.cpp
@@ -279,6 +279,10 @@ bool FExrImgMediaReaderGpu::ReadFrame(int32 FrameId, const TMap<int32, FImgMedia
 
 				// Fall back to CPU.
 				bFallBackToCPU = true;
+
+				// To make sure that Media Texture doesn't call the converter if this frame is invalid.
+				OutFrame->SampleConverter.Reset();
+
 				return FExrImgMediaReader::ReadFrame(FrameId, InMipTiles, OutFrame);
 			}
 		}
@@ -350,7 +354,7 @@ FExrImgMediaReaderGpu::EReadResult FExrImgMediaReaderGpu::ReadInChunks(uint16* B
 			FScopeLock RegionScopeLock(&CanceledFramesCriticalSection);
 			if (CanceledFrames.Remove(FrameId) > 0)
 			{
-				UE_LOG(LogImgMedia, Warning, TEXT("Reader %p: Canceling Frame %i At chunk # %i"), this, FrameId, Row);
+				UE_LOG(LogImgMedia, Verbose, TEXT("Reader %p: Canceling Frame %i At chunk # %i"), this, FrameId, Row);
 				bResult = Cancelled;
 				break;
 			}
@@ -622,7 +626,11 @@ void FExrImgMediaReaderGpu::TransferFromStagingBuffer()
 bool FExrMediaTextureSampleConverter::Convert(FTexture2DRHIRef& InDstTexture, const FConversionHints& Hints)
 {
 	FScopeLock ScopeLock(&ConverterCallbacksCriticalSection);
-	bool bExecutionSuccessful = ConvertExrBufferCallback.Execute(FRHICommandListExecutor::GetImmediateCommandList(), InDstTexture, MipBuffers);
+	bool bExecutionSuccessful = false;
+	if (ConvertExrBufferCallback.IsBound())
+	{
+		bExecutionSuccessful = ConvertExrBufferCallback.Execute(FRHICommandListExecutor::GetImmediateCommandList(), InDstTexture, MipBuffers);
+	}
 	return bExecutionSuccessful;
 }
 
