ImgMedia: Correctly resetting the calculated mip level cache between views (fixing stale tiles on planes in the nDisplay outer), also now printing the full list of visible mips on ImgMedia.MipMapDebug.

diff --git a/Engine/Plugins/Media/ImgMedia/Source/ImgMedia/Private/Assets/ImgMediaMipMapInfo.cpp b/Engine/Plugins/Media/ImgMedia/Source/ImgMedia/Private/Assets/ImgMediaMipMapInfo.cpp
--- a/Engine/Plugins/Media/ImgMedia/Source/ImgMedia/Private/Assets/ImgMediaMipMapInfo.cpp
+++ b/Engine/Plugins/Media/ImgMedia/Source/ImgMedia/Private/Assets/ImgMediaMipMapInfo.cpp
@@ -8,6 +8,7 @@
 
 #include "Async/Async.h"
 #include "Components/StaticMeshComponent.h"
+#include "Containers/Set.h"
 #include "Engine/Engine.h"
 #include "Engine/StaticMesh.h"
 #include "GameFramework/Actor.h"
@@ -335,10 +336,6 @@ namespace {
 
 			// To avoid calculating tile corner mip levels multiple times over, we cache them in this array.
 			CornerMipLevelsCached.SetNum((InSequenceInfo.NumTiles.X + 1) * (InSequenceInfo.NumTiles.Y + 1));
-			for (float& Level : CornerMipLevelsCached)
-			{
-				Level = -1.0f;
-			}
 
 			const FTransform MeshTransform = Mesh->GetComponentTransform();
 			const FVector MeshScale = Mesh->GetComponentScale();
@@ -356,6 +353,8 @@ namespace {
 					continue;
 				}
 
+				ResetMipLevelCache();
+
 				// Get frustum.
 				FConvexVolume ViewFrustum;
 				GetViewFrustumBounds(ViewFrustum, ViewInfo.ViewProjectionMatrix, false, false);
@@ -518,6 +517,15 @@ namespace {
 			CornerMipLevelsCached[Address0Y * RowSize + Address0X] = InCalculatedLevel;
 		}
 
+		/** Convenience function to reset the cache. */
+		FORCEINLINE void ResetMipLevelCache() const
+		{
+			for (float& Level : CornerMipLevelsCached)
+			{
+				Level = -1.0f;
+			}
+		}
+
 		/** Local size of this mesh component. */
 		FVector PlaneSize;
 
@@ -793,21 +801,29 @@ void FImgMediaMipMapInfo::Tick(float DeltaTime)
 	{
 		if (GEngine != nullptr)
 		{
-			FIntVector2 MipLevelRange = FIntVector2(TNumericLimits<int32>::Max(), 0);
+			TSet<int32> VisibleMips;
 			int32 NumVisibleTiles = 0;
+
 			for (const auto& MipTiles : CachedVisibleTiles)
 			{
-				MipLevelRange[0] = FMath::Min(MipLevelRange[0], MipTiles.Key);
-				MipLevelRange[1] = FMath::Max(MipLevelRange[1], MipTiles.Key);
+				VisibleMips.Add(MipTiles.Key);
 
 				const FImgMediaTileSelection& TileSelection = MipTiles.Value;
 				NumVisibleTiles += TileSelection.NumVisibleTiles();
 			}
 
-			if (MipLevelRange[0] <= MipLevelRange[1])
+			if (VisibleMips.Num() > 0)
 			{
-				GEngine->AddOnScreenDebugMessage(-1, 0.0f, FColor::Yellow, *FString::Printf(TEXT("%s Mip Level: [%d, %d]"), *SequenceInfo.Name.ToString(), MipLevelRange[0], MipLevelRange[1]));
-				GEngine->AddOnScreenDebugMessage(-1, 0.0f, FColor::Yellow, *FString::Printf(TEXT("%s Num Tiles: %d"), *SequenceInfo.Name.ToString(), NumVisibleTiles));
+				auto VisibleMipsIt = VisibleMips.CreateConstIterator();
+				FString Mips = FString::FromInt(*VisibleMipsIt);
+				
+				for (++VisibleMipsIt; VisibleMipsIt; ++VisibleMipsIt)
+				{
+					Mips += FString(TEXT(", ")) + FString::FromInt(*VisibleMipsIt);
+				}
+
+				GEngine->AddOnScreenDebugMessage(-1, 0.0f, FColor::Yellow, *FString::Printf(TEXT("%s Mip Level(s): [%s]"), *SequenceInfo.Name.ToString(), *Mips));
+				GEngine->AddOnScreenDebugMessage(-1, 0.0f, FColor::Yellow, *FString::Printf(TEXT("%s Num Tile(s): %d"), *SequenceInfo.Name.ToString(), NumVisibleTiles));
 			}
 		}
 	}
