ImgMedia: ProcessImages only uses one EXR thread now as otherwise it could use a lot of memory.

#rnx

diff --git a/Engine/Plugins/Media/ImgMedia/Source/ImgMediaEditor/Private/Widgets/SImgMediaProcessImages.cpp b/Engine/Plugins/Media/ImgMedia/Source/ImgMediaEditor/Private/Widgets/SImgMediaProcessImages.cpp
--- a/Engine/Plugins/Media/ImgMedia/Source/ImgMediaEditor/Private/Widgets/SImgMediaProcessImages.cpp
+++ b/Engine/Plugins/Media/ImgMedia/Source/ImgMediaEditor/Private/Widgets/SImgMediaProcessImages.cpp
@@ -516,7 +516,7 @@ void SImgMediaProcessImages::ProcessImageCustomRawData(TArray64<uint8>& RawData,
 	}
 
 	// Create output.
-	OutFile.CreateOutputFile(InName, DestWidth, DestHeight, bInEnableMips);
+	OutFile.CreateOutputFile(InName, DestWidth, DestHeight, bInEnableMips, 1);
 	if (DestNumChannels == 4)
 	{
 		OutFile.AddFrameBufferChannel(AChannelName, nullptr, Stride);
diff --git a/Engine/Plugins/Media/ImgMedia/Source/OpenExrWrapper/Private/OpenExrWrapper.cpp b/Engine/Plugins/Media/ImgMedia/Source/OpenExrWrapper/Private/OpenExrWrapper.cpp
--- a/Engine/Plugins/Media/ImgMedia/Source/OpenExrWrapper/Private/OpenExrWrapper.cpp
+++ b/Engine/Plugins/Media/ImgMedia/Source/OpenExrWrapper/Private/OpenExrWrapper.cpp
@@ -426,7 +426,7 @@ void FTiledOutputFile::AddChannel(const FString& Name)
 }
 
 void FTiledOutputFile::CreateOutputFile(const FString& FilePath,
-	int32 TileWidth, int32 TileHeight, bool bIsMipsEnabled)
+	int32 TileWidth, int32 TileHeight, bool bIsMipsEnabled, int32 NumThreads)
 {
 	if (OutputFile == nullptr)
 	{
@@ -438,7 +438,7 @@ void FTiledOutputFile::CreateOutputFile(const FString& FilePath,
 
 			// Create output file.
 			OutputFile = new Imf::TiledOutputFile(StringCast<ANSICHAR>(*FilePath).Get(),
-				*((Imf::Header*)Header));
+				*((Imf::Header*)Header), NumThreads);
 		}
 		catch (std::exception const& Exception)
 		{
diff --git a/Engine/Plugins/Media/ImgMedia/Source/OpenExrWrapper/Public/OpenExrWrapper.h b/Engine/Plugins/Media/ImgMedia/Source/OpenExrWrapper/Public/OpenExrWrapper.h
--- a/Engine/Plugins/Media/ImgMedia/Source/OpenExrWrapper/Public/OpenExrWrapper.h
+++ b/Engine/Plugins/Media/ImgMedia/Source/OpenExrWrapper/Public/OpenExrWrapper.h
@@ -239,9 +239,10 @@ public:
 	 * @param TileWidth			Width of a tile.
 	 * @param TileHeight		Height of a tile.
 	 * @param bIsMipsEnabled	True to enable mip mapping.
+	 * @param NumThreads		Number of threads for EXR to use.
 	 */
 	void CreateOutputFile(const FString& FilePath,
-		int32 TileWidth, int32 TileHeight, bool bIsMipsEnabled);
+		int32 TileWidth, int32 TileHeight, bool bIsMipsEnabled, int32 NumThreads);
 
 	/**
 	 * Get the width of a mip level.
