ImgMedia: Removing shelved VT approach for visible tile calculations. (Note that the tiling IMediaView interface could still prove useful for rehearsal so we leave it present but unimplemented.)

diff --git a/Engine/Plugins/Media/ImgMedia/Source/ImgMedia/Private/Assets/ImgMediaSource.cpp b/Engine/Plugins/Media/ImgMedia/Source/ImgMedia/Private/Assets/ImgMediaSource.cpp
--- a/Engine/Plugins/Media/ImgMedia/Source/ImgMedia/Private/Assets/ImgMediaSource.cpp
+++ b/Engine/Plugins/Media/ImgMedia/Source/ImgMedia/Private/Assets/ImgMediaSource.cpp
@@ -18,7 +18,6 @@ UImgMediaSource::UImgMediaSource()
 	: IsPathRelativeToProjectRoot(false)
 	, FrameRateOverride(0, 0)
 	, bFillGapsInSequence(true)
-	, bReadVirtualTextureTiles(false)
 	, MipMapInfo(MakeShared<FImgMediaMipMapInfo, ESPMode::ThreadSafe>())
 {
 }
@@ -79,10 +78,6 @@ bool UImgMediaSource::GetMediaOption(const FName& Key, bool DefaultValue) const
 	{
 		return bFillGapsInSequence;
 	}
-	if (Key == ImgMedia::ReadVirtualTextureTiles)
-	{
-		return bReadVirtualTextureTiles;
-	}
 
 	return Super::GetMediaOption(Key, DefaultValue);
 }
@@ -129,8 +124,7 @@ bool UImgMediaSource::HasMediaOption(const FName& Key) const
 		(Key == ImgMedia::FrameRateOverrideDenonimatorOption) ||
 		(Key == ImgMedia::FrameRateOverrideNumeratorOption) ||
 		(Key == ImgMedia::ProxyOverrideOption) ||
-		(Key == ImgMedia::MipMapInfoOption) ||
-		(Key == ImgMedia::ReadVirtualTextureTiles))
+		(Key == ImgMedia::MipMapInfoOption))
 	{
 		return true;
 	}
diff --git a/Engine/Plugins/Media/ImgMedia/Source/ImgMedia/Private/ImgMediaPrivate.h b/Engine/Plugins/Media/ImgMedia/Source/ImgMedia/Private/ImgMediaPrivate.h
--- a/Engine/Plugins/Media/ImgMedia/Source/ImgMedia/Private/ImgMediaPrivate.h
+++ b/Engine/Plugins/Media/ImgMedia/Source/ImgMedia/Private/ImgMediaPrivate.h
@@ -55,9 +55,6 @@ namespace ImgMedia
 	/** Name of the NumTilesY media option. */
 	static const FName NumTilesYOption("ImgNumTilesY");
 
-	/** Name of the ReadVirtualTextureTiles media option. */
-	static const FName ReadVirtualTextureTiles("ImgReadVirtualTextureTiles");
-
 	/** Name of the smart cache media option. */
 	static const FName SmartCacheEnabled("ImgMediaSmartCacheEnabled");
 
diff --git a/Engine/Plugins/Media/ImgMedia/Source/ImgMedia/Private/Loader/ImgMediaLoader.cpp b/Engine/Plugins/Media/ImgMedia/Source/ImgMedia/Private/Loader/ImgMediaLoader.cpp
--- a/Engine/Plugins/Media/ImgMedia/Source/ImgMedia/Private/Loader/ImgMediaLoader.cpp
+++ b/Engine/Plugins/Media/ImgMedia/Source/ImgMedia/Private/Loader/ImgMediaLoader.cpp
@@ -85,13 +85,11 @@
 	const TSharedRef<FImgMediaGlobalCache, ESPMode::ThreadSafe>& InGlobalCache,
 	const TSharedPtr<FImgMediaMipMapInfo, ESPMode::ThreadSafe>& InMipMapInfo,
 	bool bInFillGapsInSequence,
-	bool bInReadVirtualTextureTiles,
 	const FImgMediaLoaderSmartCacheSettings& InSmartCacheSettings)
 	: Frames(1)
 	, ImageWrapperModule(FModuleManager::LoadModuleChecked<IImageWrapperModule>("ImageWrapper"))
 	, Initialized(false)
 	, bFillGapsInSequence(bInFillGapsInSequence)
-	, bReadVirtualTextureTiles(bInReadVirtualTextureTiles)
 	, bIsTiled(false)
 	, TilingDescription()
 	, NumLoadAhead(0)
@@ -653,6 +651,11 @@
 	TMap<int32, FImgMediaTileSelection> DesiredMipsAndTiles;
 	GetDesiredMipTiles(FrameNumber, DesiredMipsAndTiles);
 	
+	if (DesiredMipsAndTiles.IsEmpty())
+	{
+		return nullptr;
+	}
+	
 	// Set up work.
 	Work->Initialize(FrameNumber, DesiredMipsAndTiles, ExistingFrame);
 	QueuedFrameNumbers.Add(FrameNumber);
@@ -1301,24 +1304,8 @@
 
 void FImgMediaLoader::GetDesiredMipTiles(int32 FrameIndex, TMap<int32, FImgMediaTileSelection>& OutMipsAndTiles)
 {
-	if(bReadVirtualTextureTiles)
+	if(MipMapInfo.IsValid() && MipMapInfo->HasObjects())
 	{
-		TMap<int32, TSet<FMediaTileCoordinate>> VisibleTiles;
-		GetVisibleTiles(VisibleTiles);
-
-		for (const auto& Pair : VisibleTiles)
-		{
-			FImgMediaTileSelection Selection = FImgMediaTileSelection::CreateForTargetMipLevel(GetNumTilesX(), GetNumTilesY(), Pair.Key);
-
-			for (const FMediaTileCoordinate& Coordinate : Pair.Value)
-			{
-				Selection.SetVisible(Coordinate.X, Coordinate.Y);
-			}
-			OutMipsAndTiles.Emplace(Pair.Key, MoveTemp(Selection));
-		}
-	}
-	else if(MipMapInfo.IsValid() && MipMapInfo->HasObjects())
-	{
 		OutMipsAndTiles = MipMapInfo->GetVisibleTiles();
 	}
 	else
@@ -1414,18 +1401,3 @@
 	return bFoundNumber;
 }
 
-bool FImgMediaLoader::GetVisibleTiles(TMap<int32, TSet<FMediaTileCoordinate>>& OutTiles) const
-{
-	FScopeLock Lock(&CriticalSectionTiles);
-	OutTiles = ActiveTilesPerMipLevel;
-
-	return true;
-}
-
-bool FImgMediaLoader::SetVisibleTiles(TMap<int32, TSet<FMediaTileCoordinate>>&& InTiles)
-{
-	FScopeLock Lock(&CriticalSectionTiles);
-	ActiveTilesPerMipLevel = MoveTemp(InTiles);
-
-	return true;
-}
diff --git a/Engine/Plugins/Media/ImgMedia/Source/ImgMedia/Private/Loader/ImgMediaLoader.h b/Engine/Plugins/Media/ImgMedia/Source/ImgMedia/Private/Loader/ImgMediaLoader.h
--- a/Engine/Plugins/Media/ImgMedia/Source/ImgMedia/Private/Loader/ImgMediaLoader.h
+++ b/Engine/Plugins/Media/ImgMedia/Source/ImgMedia/Private/Loader/ImgMediaLoader.h
@@ -58,7 +58,6 @@ public:
 		const TSharedRef<FImgMediaGlobalCache, ESPMode::ThreadSafe>& InGlobalCache,
 		const TSharedPtr<FImgMediaMipMapInfo, ESPMode::ThreadSafe>& InMipMapInfo,
 		bool bInFillGapsInSequence,
-		bool bInReadVirtualTextureTiles,
 		const FImgMediaLoaderSmartCacheSettings& InSmartCacheSettings);
 
 	/** Virtual destructor. */
@@ -461,14 +460,6 @@ protected:
 	 */
 	bool GetNumberAtEndOfString(int32& Number, const FString& String) const;
 
-public:
-	
-	/** Thread-safe function to get currently visible tiles per mip level. */
-	bool GetVisibleTiles(TMap<int32, TSet<FMediaTileCoordinate>>& OutTiles) const;
-
-	/** Thread-safe function to set currently visible tiles per mip level. */
-	bool SetVisibleTiles(TMap<int32, TSet<FMediaTileCoordinate>>&& InTiles);
-
 private:
 
 	/** Critical section for synchronizing access to Frames. */
@@ -495,9 +486,6 @@ private:
 	/** If true, then any gaps in the sequence will be filled with blank frames. */
 	const bool bFillGapsInSequence;
 
-	/** If true, then virtual texture tiles will be used instead of the MipMapInfo object (applicable for tiled image sequences). */
-	const bool bReadVirtualTextureTiles;
-
 	/** True if this sequence contains tiles. */
 	bool bIsTiled;
 
@@ -567,11 +555,6 @@ private:
 	} QueuedSampleFetch;
 
 private:
-	/** Critical section for synchronizing access to the active tile list. */
-	mutable FCriticalSection CriticalSectionTiles;
-
-	/** List of visible tiles updated through the IMediaView interface. */
-	TMap<int32, TSet<FMediaTileCoordinate>> ActiveTilesPerMipLevel;
 
 	/** Settings for the smart cache. */
 	FImgMediaLoaderSmartCacheSettings SmartCacheSettings;
diff --git a/Engine/Plugins/Media/ImgMedia/Source/ImgMedia/Private/Player/ImgMediaPlayer.cpp b/Engine/Plugins/Media/ImgMedia/Source/ImgMedia/Private/Player/ImgMediaPlayer.cpp
--- a/Engine/Plugins/Media/ImgMedia/Source/ImgMedia/Private/Player/ImgMediaPlayer.cpp
+++ b/Engine/Plugins/Media/ImgMedia/Source/ImgMedia/Private/Player/ImgMediaPlayer.cpp
@@ -229,7 +229,6 @@ bool FImgMediaPlayer::Open(const FString& Url, const IMediaOptions* Options)
 	FFrameRate FrameRateOverride(0, 0);
 	TSharedPtr<FImgMediaMipMapInfo, ESPMode::ThreadSafe> MipMapInfo;
 	bool bFillGapsInSequence = true;
-	bool bReadVirtualTextureTiles = false;
 	bool bIsSmartCacheEnabled = false;
 	float SmartCacheTimeToLookAhead = 0.0f;
 	if (Options != nullptr)
@@ -237,7 +236,6 @@ bool FImgMediaPlayer::Open(const FString& Url, const IMediaOptions* Options)
 		FrameRateOverride.Denominator = Options->GetMediaOption(ImgMedia::FrameRateOverrideDenonimatorOption, 0LL);
 		FrameRateOverride.Numerator = Options->GetMediaOption(ImgMedia::FrameRateOverrideNumeratorOption, 0LL);
 		bFillGapsInSequence = Options->GetMediaOption(ImgMedia::FillGapsInSequenceOption, true);
-		bReadVirtualTextureTiles = Options->GetMediaOption(ImgMedia::ReadVirtualTextureTiles, false);
 		bIsSmartCacheEnabled = Options->GetMediaOption(ImgMedia::SmartCacheEnabled, false);
 		SmartCacheTimeToLookAhead = Options->GetMediaOption(ImgMedia::SmartCacheTimeToLookAhead, 0.0f);
 		TSharedPtr<IMediaOptions::FDataContainer, ESPMode::ThreadSafe> DefaultValue;
@@ -258,8 +256,7 @@ bool FImgMediaPlayer::Open(const FString& Url, const IMediaOptions* Options)
 	// initialize image loader on a separate thread
 	FImgMediaLoaderSmartCacheSettings SmartCacheSettings(bIsSmartCacheEnabled, SmartCacheTimeToLookAhead);
 	Loader = MakeShared<FImgMediaLoader, ESPMode::ThreadSafe>(Scheduler.ToSharedRef(),
-		GlobalCache.ToSharedRef(), MipMapInfo, bFillGapsInSequence, bReadVirtualTextureTiles,
-		SmartCacheSettings);
+		GlobalCache.ToSharedRef(), MipMapInfo, bFillGapsInSequence, SmartCacheSettings);
 	Scheduler->RegisterLoader(Loader.ToSharedRef());
 
 	const FString SequencePath = Url.RightChop(6);
@@ -457,29 +454,6 @@ bool FImgMediaPlayer::GetPlayerFeatureFlag(EFeatureFlag flag) const
 	return IMediaPlayer::GetPlayerFeatureFlag(flag);
 }
 
-/* IMediaView interface
- *****************************************************************************/
-
-bool FImgMediaPlayer::GetVisibleTiles(TMap<int32, TSet<FMediaTileCoordinate>>& OutTiles) const
-{
-	if (Loader.IsValid())
-	{
-		return Loader->GetVisibleTiles(OutTiles);
-	}
-
-	return false;
-}
-
-bool FImgMediaPlayer::SetVisibleTiles(TMap<int32, TSet<FMediaTileCoordinate>>&& InTiles)
-{
-	if (Loader.IsValid())
-	{
-		return Loader->SetVisibleTiles(MoveTemp(InTiles));
-	}
-
-	return false;
-}
-
 
 /* FImgMediaPlayer implementation
  *****************************************************************************/
diff --git a/Engine/Plugins/Media/ImgMedia/Source/ImgMedia/Private/Player/ImgMediaPlayer.h b/Engine/Plugins/Media/ImgMedia/Source/ImgMedia/Private/Player/ImgMediaPlayer.h
--- a/Engine/Plugins/Media/ImgMedia/Source/ImgMedia/Private/Player/ImgMediaPlayer.h
+++ b/Engine/Plugins/Media/ImgMedia/Source/ImgMedia/Private/Player/ImgMediaPlayer.h
@@ -68,13 +68,6 @@ public:
 	virtual void ProcessVideoSamples() override;
 	virtual bool GetPlayerFeatureFlag(EFeatureFlag flag) const override;
 
-public:
-
-	//~ IMediaView interface
-
-	virtual bool GetVisibleTiles(TMap<int32, TSet<FMediaTileCoordinate>>& OutTiles) const override;
-	virtual bool SetVisibleTiles(TMap<int32, TSet<FMediaTileCoordinate>>&& InTiles) override;
-
 protected:
 
 	/**
diff --git a/Engine/Plugins/Media/ImgMedia/Source/ImgMedia/Public/ImgMediaSource.h b/Engine/Plugins/Media/ImgMedia/Source/ImgMedia/Public/ImgMediaSource.h
--- a/Engine/Plugins/Media/ImgMedia/Source/ImgMedia/Public/ImgMediaSource.h
+++ b/Engine/Plugins/Media/ImgMedia/Source/ImgMedia/Public/ImgMediaSource.h
@@ -66,10 +66,6 @@ public:
 	UPROPERTY(BlueprintReadWrite, EditAnywhere, Category = Sequence)
 	bool bFillGapsInSequence;
 
-	/** If true, tiled image sequences will rely only on visible tiles provided by the virtual texture system for playback. */
-	UPROPERTY(BlueprintReadWrite, EditAnywhere, Category = Sequence, AdvancedDisplay)
-	bool bReadVirtualTextureTiles;
-
 #if WITH_EDITORONLY_DATA
 
 	/** This is only used so we can customize editing of SequencePath. */
diff --git a/Engine/Plugins/Media/ImgMedia/Source/ImgMediaEditor/Private/Widgets/ImgMediaProcessImagesOptions.h b/Engine/Plugins/Media/ImgMedia/Source/ImgMediaEditor/Private/Widgets/ImgMediaProcessImagesOptions.h
--- a/Engine/Plugins/Media/ImgMedia/Source/ImgMediaEditor/Private/Widgets/ImgMediaProcessImagesOptions.h
+++ b/Engine/Plugins/Media/ImgMedia/Source/ImgMediaEditor/Private/Widgets/ImgMediaProcessImagesOptions.h
@@ -43,10 +43,6 @@ public:
 	UPROPERTY(EditAnywhere, Transient, Category = Tiles, meta = (EditCondition = "bEnableTiling", ClampMin = "0.0"))
 	int32 TileSizeY = 256;
 
-	/** Number of pixels to duplicate along each tile border. */
-	UPROPERTY(EditAnywhere, Transient, Category = Tiles, meta = (EditCondition = "bEnableTiling", ClampMin = "0.0"))
-	int32 TileBorder = 0;
-
 	/** Number of threads to use when processing. If 0 then this number is set automatically. */
 	UPROPERTY(EditAnywhere, Transient, Category = Processing, meta = (ClampMin = "0.0"))
 	int32 NumThreads = 0;
diff --git a/Engine/Plugins/Media/ImgMedia/Source/ImgMediaEditor/Private/Widgets/SImgMediaProcessImages.cpp b/Engine/Plugins/Media/ImgMedia/Source/ImgMediaEditor/Private/Widgets/SImgMediaProcessImages.cpp
--- a/Engine/Plugins/Media/ImgMedia/Source/ImgMediaEditor/Private/Widgets/SImgMediaProcessImages.cpp
+++ b/Engine/Plugins/Media/ImgMedia/Source/ImgMediaEditor/Private/Widgets/SImgMediaProcessImages.cpp
@@ -200,7 +200,7 @@ void SImgMediaProcessImages::ProcessAllImages()
 	bool bUseCustomFormat = Options->bUseCustomFormat;
 	int32 InTileWidth = Options->bEnableTiling ? Options->TileSizeX : 0;
 	int32 InTileHeight = Options->bEnableTiling ? Options->TileSizeY : 0;
-	int32 TileBorder = Options->bEnableTiling ? Options->TileBorder : 0;
+	int32 TileBorder = 0; // Note: virtual texture support is shelved for now.
 	bool bEnableMips = Options->bEnableMipMapping;
 
 	// Create output directory.
@@ -912,7 +912,7 @@ void SImgMediaProcessImages::HandleProcessing()
 					bool bUseCustomFormat = Options->bUseCustomFormat;
 					int32 InTileWidth = Options->bEnableTiling ? Options->TileSizeX : 0;
 					int32 InTileHeight = Options->bEnableTiling ? Options->TileSizeY : 0;
-					int32 TileBorder = Options->bEnableTiling ? Options->TileBorder : 0;
+					int32 TileBorder = 0; // Note: virtual texture support is shelved for now.
 					bool bEnableMips = Options->bEnableMipMapping;
 					bool bHasAlphaChannel = false;
 					FString OutPath = Options->OutputPath.Path;
