ImgMedia: Adding FOV multiplier for mip/tiles calculations as a potential mitigation for missing tiles with moving cameras. While nDisplay already supports overscan, further expansion may be needed to prevent any artifact in the inner.

diff --git a/Engine/Plugins/Media/ImgMedia/Source/ImgMedia/Private/Assets/ImgMediaMipMapInfo.cpp b/Engine/Plugins/Media/ImgMedia/Source/ImgMedia/Private/Assets/ImgMediaMipMapInfo.cpp
--- a/Engine/Plugins/Media/ImgMedia/Source/ImgMedia/Private/Assets/ImgMediaMipMapInfo.cpp
+++ b/Engine/Plugins/Media/ImgMedia/Source/ImgMedia/Private/Assets/ImgMediaMipMapInfo.cpp
@@ -357,7 +357,7 @@ namespace {
 
 				// Get frustum.
 				FConvexVolume ViewFrustum;
-				GetViewFrustumBounds(ViewFrustum, ViewInfo.ViewProjectionMatrix, false, false);
+				GetViewFrustumBounds(ViewFrustum, ViewInfo.OverscanViewProjectionMatrix, false, false);
 
 				int32 MaxLevel = InSequenceInfo.NumMipLevels - 1;
 				int MipLevelDiv = 1 << MaxLevel;
@@ -557,7 +557,7 @@ namespace {
 
 				// Analytical derivation of visible tiles from the view frustum, given a sphere presumed to be infinitely large
 				FConvexVolume ViewFrustum;
-				GetViewFrustumBounds(ViewFrustum, ViewInfo.ViewProjectionMatrix, false, false);
+				GetViewFrustumBounds(ViewFrustum, ViewInfo.OverscanViewProjectionMatrix, false, false);
 				
 				// Include all tiles containted in the visible UV region
 				int32 NumX = InSequenceInfo.NumTiles.X;
diff --git a/Engine/Plugins/Media/ImgMedia/Source/ImgMedia/Private/ImgMediaSceneViewExtension.cpp b/Engine/Plugins/Media/ImgMedia/Source/ImgMedia/Private/ImgMediaSceneViewExtension.cpp
--- a/Engine/Plugins/Media/ImgMedia/Source/ImgMedia/Private/ImgMediaSceneViewExtension.cpp
+++ b/Engine/Plugins/Media/ImgMedia/Source/ImgMedia/Private/ImgMediaSceneViewExtension.cpp
@@ -5,6 +5,11 @@
 
 #include "SceneView.h"
 
+static TAutoConsoleVariable<float> CVarImgMediaFieldOfViewMultiplier(
+	TEXT("ImgMedia.FieldOfViewMultiplier"),
+	1.0f,
+	TEXT("Multiply the field of view for active cameras by this value, generally to increase the frustum overall sizes to mitigate missing tile artifacts.\n"),
+	ECVF_Default);
 
 FImgMediaSceneViewExtension::FImgMediaSceneViewExtension(const FAutoRegister& AutoReg)
 	: FSceneViewExtensionBase(AutoReg)
@@ -39,13 +44,31 @@
 
 	static const auto CVarMinAutomaticViewMipBiasOffset = IConsoleManager::Get().FindTConsoleVariableDataFloat(TEXT("r.ViewTextureMipBias.Offset"));
 	static const auto CVarMinAutomaticViewMipBias = IConsoleManager::Get().FindTConsoleVariableDataFloat(TEXT("r.ViewTextureMipBias.Min"));
+	const float FieldOfViewMultiplier = CVarImgMediaFieldOfViewMultiplier.GetValueOnGameThread();
 
 	for (const FSceneView* View : InViewFamily.Views)
 	{
 		FImgMediaViewInfo Info;
 		Info.Location = View->ViewMatrices.GetViewOrigin();
-		Info.ViewMatrix = View->ViewMatrices.GetViewMatrix();
 		Info.ViewProjectionMatrix = View->ViewMatrices.GetViewProjectionMatrix();
+
+		if (FMath::IsNearlyEqual(FieldOfViewMultiplier, 1.0f))
+		{
+			Info.OverscanViewProjectionMatrix = Info.ViewProjectionMatrix;
+		}
+		else
+		{
+			FMatrix AdjustedProjectionMatrix = View->ViewMatrices.GetProjectionMatrix();
+
+			const double HalfHorizontalFOV = FMath::Atan(1.0 / AdjustedProjectionMatrix.M[0][0]);
+			const double HalfVerticalFOV = FMath::Atan(1.0 / AdjustedProjectionMatrix.M[1][1]);
+
+			AdjustedProjectionMatrix.M[0][0] = 1.0 / FMath::Tan(HalfHorizontalFOV * FieldOfViewMultiplier);
+			AdjustedProjectionMatrix.M[1][1] = 1.0 / FMath::Tan(HalfVerticalFOV * FieldOfViewMultiplier);
+			
+			Info.OverscanViewProjectionMatrix = View->ViewMatrices.GetViewMatrix() * AdjustedProjectionMatrix;
+		}
+		
 		Info.ViewportRect = View->UnconstrainedViewRect.Scale(ResolutionFraction);
 
 		// We store hidden or show-only ids to later avoid needless calculations when objects are not in view.
diff --git a/Engine/Plugins/Media/ImgMedia/Source/ImgMedia/Public/ImgMediaSceneViewExtension.h b/Engine/Plugins/Media/ImgMedia/Source/ImgMedia/Public/ImgMediaSceneViewExtension.h
--- a/Engine/Plugins/Media/ImgMedia/Source/ImgMedia/Public/ImgMediaSceneViewExtension.h
+++ b/Engine/Plugins/Media/ImgMedia/Source/ImgMedia/Public/ImgMediaSceneViewExtension.h
@@ -14,10 +14,10 @@ struct IMGMEDIA_API FImgMediaViewInfo
 {
 	/** Position of camera. */
 	FVector Location;
-	/** View matrix of camera. */
-	FMatrix ViewMatrix;
-	/** View projection matrix of camera. */
+	/** View-projection matrix of the camera. */
 	FMatrix ViewProjectionMatrix;
+	/** View-projection matrix of the camera, optionally scaled for overscan frustum calculations. */
+	FMatrix OverscanViewProjectionMatrix;
 	/** Active viewport size. */
 	FIntRect ViewportRect;
 	/** View mip bias. */
