ImgMedia: Backing out of incorrect/leaking attempt at bypassing image loader work in 20560079.

diff --git a/Engine/Plugins/Media/ImgMedia/Source/ImgMedia/Private/Loader/ImgMediaLoader.cpp b/Engine/Plugins/Media/ImgMedia/Source/ImgMedia/Private/Loader/ImgMediaLoader.cpp
--- a/Engine/Plugins/Media/ImgMedia/Source/ImgMedia/Private/Loader/ImgMediaLoader.cpp
+++ b/Engine/Plugins/Media/ImgMedia/Source/ImgMedia/Private/Loader/ImgMediaLoader.cpp
@@ -651,11 +651,6 @@ IQueuedWork* FImgMediaLoader::GetWork()
 	TMap<int32, FImgMediaTileSelection> DesiredMipsAndTiles;
 	GetDesiredMipTiles(FrameNumber, DesiredMipsAndTiles);
 	
-	if (DesiredMipsAndTiles.IsEmpty())
-	{
-		return nullptr;
-	}
-	
 	// Set up work.
 	Work->Initialize(FrameNumber, DesiredMipsAndTiles, ExistingFrame);
 	QueuedFrameNumbers.Add(FrameNumber);
