ImgMedia: Fix destruction crash related to using a (sequencer) spawnable media plate. The mesh component weak object pointer would be pending kill, and therefore return null on get.

#rnx

diff --git a/Engine/Plugins/Media/ImgMedia/Source/ImgMedia/Private/Assets/ImgMediaMipMapInfo.cpp b/Engine/Plugins/Media/ImgMedia/Source/ImgMedia/Private/Assets/ImgMediaMipMapInfo.cpp
--- a/Engine/Plugins/Media/ImgMedia/Source/ImgMedia/Private/Assets/ImgMediaMipMapInfo.cpp
+++ b/Engine/Plugins/Media/ImgMedia/Source/ImgMedia/Private/Assets/ImgMediaMipMapInfo.cpp
@@ -259,7 +259,7 @@ FImgMediaMipMapObjectInfo::FImgMediaMipMapObjectInfo(UMeshComponent* InMeshCompo
 
 UMeshComponent* FImgMediaMipMapObjectInfo::GetMeshComponent() const
 {
-	return MeshComponent.Get();
+	return MeshComponent.Get(true);
 }
 
 void FImgMediaMipMapObjectInfo::CalculateVisibleTiles(const TArray<FImgMediaViewInfo>& InViewInfos, const FSequenceInfo& InSequenceInfo, TMap<int32, FImgMediaTileSelection>& VisibleTiles) const
@@ -675,11 +675,15 @@ void FImgMediaMipMapInfo::RemoveObject(AActor* InActor)
 		for (int Index = 0; Index < Objects.Num(); ++Index)
 		{
 			FImgMediaMipMapObjectInfo* Info = Objects[Index];
-			if (InActor == Info->GetMeshComponent()->GetOuter())
+
+			if (UMeshComponent* MeshComponent = Info->GetMeshComponent())
 			{
-				Objects.RemoveAtSwap(Index);
-				delete Info;
-				break;
+				if (InActor == MeshComponent->GetOuter())
+				{
+					Objects.RemoveAtSwap(Index);
+					delete Info;
+					break;
+				}
 			}
 		}
 	}
