Removed OpenExrWrapper private pch.
Timing when compiling module with PCH(including compiling PCH): 3.07 secs
Timing using shared PCH: 1.53 secs

diff --git a/Engine/Plugins/Media/ImgMedia/Source/OpenExrWrapper/OpenExrWrapper.Build.cs b/Engine/Plugins/Media/ImgMedia/Source/OpenExrWrapper/OpenExrWrapper.Build.cs
--- a/Engine/Plugins/Media/ImgMedia/Source/OpenExrWrapper/OpenExrWrapper.Build.cs
+++ b/Engine/Plugins/Media/ImgMedia/Source/OpenExrWrapper/OpenExrWrapper.Build.cs
@@ -7,11 +7,7 @@ namespace UnrealBuildTool.Rules
 		public OpenExrWrapper(ReadOnlyTargetRules Target) : base(Target)
 		{
 			bEnableExceptions = true;
-
-			PCHUsage = ModuleRules.PCHUsageMode.NoSharedPCHs;
-			PrivatePCHHeaderFile = "Public/OpenExrWrapper.h";
-
-            bUseRTTI = true;
+			bUseRTTI = true;
 
             PublicDependencyModuleNames.AddRange(
 				new string[] {
@@ -34,7 +30,7 @@ namespace UnrealBuildTool.Rules
                 System.Console.WriteLine("OpenExrWrapper does not supported this platform");
             }
 
-            PrivateIncludePaths.Add("OpenExrWrapper/Private");
+			PrivateIncludePaths.Add("OpenExrWrapper/Private");
         }
 	}
 }
