ImgMedia: Fixed media not playing when look ahead time is less than the frame time.

#rnx

diff --git a/Engine/Plugins/Media/ImgMedia/Source/ImgMedia/Private/Loader/ImgMediaLoader.cpp b/Engine/Plugins/Media/ImgMedia/Source/ImgMedia/Private/Loader/ImgMediaLoader.cpp
--- a/Engine/Plugins/Media/ImgMedia/Source/ImgMedia/Private/Loader/ImgMediaLoader.cpp
+++ b/Engine/Plugins/Media/ImgMedia/Source/ImgMedia/Private/Loader/ImgMediaLoader.cpp
@@ -879,7 +879,10 @@ void FImgMediaLoader::LoadSequence(const FString& SequencePath, const FFrameRate
 		float TimeToLookAhead = SmartCacheSettings.TimeToLookAhead;
 		
 		NumFramesToLoad = TimeToLookAhead / SequenceFrameRate.AsInterval();
-		NumFramesToLoad = FMath::Clamp(NumFramesToLoad, 0, GetNumImages());
+		// Clamp min frame to load to 1 so we always try to load something.
+		int32 NumImages = GetNumImages();
+		int32 MinFrames = FMath::Min(1, NumImages);
+		NumFramesToLoad = FMath::Clamp(NumFramesToLoad, MinFrames, NumImages);
 		NumLoadBehind = 0;
 		NumLoadAhead = NumFramesToLoad;
 	}
