ImgMedia: Support for partial sphere calculations.Also removed visibility/editing for the Media Plate Component MeshHorizontalRange property since its control is only valid through the sphere mesh interface.

diff --git a/Engine/Plugins/Media/ImgMedia/Source/ImgMedia/Private/Assets/ImgMediaMipMapInfo.cpp b/Engine/Plugins/Media/ImgMedia/Source/ImgMedia/Private/Assets/ImgMediaMipMapInfo.cpp
--- a/Engine/Plugins/Media/ImgMedia/Source/ImgMedia/Private/Assets/ImgMediaMipMapInfo.cpp
+++ b/Engine/Plugins/Media/ImgMedia/Source/ImgMedia/Private/Assets/ImgMediaMipMapInfo.cpp
@@ -536,7 +536,11 @@ namespace {
 	class FSphereObjectInfo : public FImgMediaMipMapObjectInfo
 	{
 	public:
-		using FImgMediaMipMapObjectInfo::FImgMediaMipMapObjectInfo;
+		FSphereObjectInfo(UMeshComponent* InMeshComponent, float InLODBias, float InMeshHorizontalRange)
+			: FImgMediaMipMapObjectInfo(InMeshComponent, InLODBias)
+			, MeshHorizontalRange(InMeshHorizontalRange)
+		{
+		}
 
 		void CalculateVisibleTiles(const TArray<FImgMediaViewInfo>& InViewInfos, const FSequenceInfo& InSequenceInfo, TMap<int32, FImgMediaTileSelection>& VisibleTiles) const override
 		{
@@ -569,10 +573,7 @@ namespace {
 						FVector2D TileCornerUV = FVector2D((float)TileX / NumX, (float)TileY / NumY);
 						
 						// Convert from latlong UV to spherical coordinates
-						FVector2D TileCornerSpherical = FVector2D(UE_PI * TileCornerUV.Y, UE_TWO_PI * TileCornerUV.X);
-						
-						// Adjust spherical coordinates to default sphere UVs
-						TileCornerSpherical.Y = -TileCornerSpherical.Y;
+						FVector2D TileCornerSpherical = FVector2D(UE_PI * TileCornerUV.Y, FMath::DegreesToRadians(MeshHorizontalRange) * (1.0f - TileCornerUV.X));
 
 						FVector TileCorner = TileCornerSpherical.SphericalToUnitCartesian() * DefaultSphereRadius;
 						TileCorner = Mesh->GetComponentTransform().TransformPosition(TileCorner);
@@ -631,6 +632,9 @@ namespace {
 				}
 			}
 		}
+
+	private:
+		float MeshHorizontalRange;
 	};
 
 } //end anonymous namespace
@@ -645,7 +649,7 @@ FImgMediaMipMapInfo::~FImgMediaMipMapInfo()
 	ClearAllObjects();
 }
 
-void FImgMediaMipMapInfo::AddObject(AActor* InActor, float Width, float LODBias, EMediaTextureVisibleMipsTiles MeshType)
+void FImgMediaMipMapInfo::AddObject(AActor* InActor, float LODBias, EMediaTextureVisibleMipsTiles MeshType, float MeshHorizontalRange)
 {
 	if (InActor != nullptr)
 	{
@@ -658,7 +662,7 @@ void FImgMediaMipMapInfo::AddObject(AActor* InActor, float Width, float LODBias,
 				Objects.Add(new FPlaneObjectInfo(MeshComponent, LODBias));
 				break;
 			case EMediaTextureVisibleMipsTiles::Sphere:
-				Objects.Add(new FSphereObjectInfo(MeshComponent, LODBias));
+				Objects.Add(new FSphereObjectInfo(MeshComponent, LODBias, MeshHorizontalRange));
 				break;
 			default:
 				Objects.Add(new FImgMediaMipMapObjectInfo(MeshComponent, LODBias));
@@ -704,7 +708,7 @@ void FImgMediaMipMapInfo::AddObjectsUsingThisMediaTexture(UMediaTexture* InMedia
 				AActor* Owner = ObjectInfo->Object.Get();
 				if (Owner != nullptr)
 				{
-					AddObject(Owner, 0.0f, ObjectInfo->MipMapLODBias, ObjectInfo->VisibleMipsTilesCalculations);
+					AddObject(Owner,ObjectInfo->MipMapLODBias, ObjectInfo->VisibleMipsTilesCalculations, ObjectInfo->MeshHorizontalRange );
 				}
 			}
 		}
diff --git a/Engine/Plugins/Media/ImgMedia/Source/ImgMedia/Private/Assets/ImgMediaMipMapInfo.h b/Engine/Plugins/Media/ImgMedia/Source/ImgMedia/Private/Assets/ImgMediaMipMapInfo.h
--- a/Engine/Plugins/Media/ImgMedia/Source/ImgMedia/Private/Assets/ImgMediaMipMapInfo.h
+++ b/Engine/Plugins/Media/ImgMedia/Source/ImgMedia/Private/Assets/ImgMediaMipMapInfo.h
@@ -211,9 +211,11 @@ public:
 	 * This object is using our img sequence.
 	 *
 	 * @param InActor Object using our img sequence.
-	 * @param Width Width of the object. If < 0, then get the width automatically.
+	 * @param LODBias Bias used for mip-level calculations (Note: must be mirrored in a material manually).
+	 * @param MeshType Plane, sphere or none mesh type for mip-tile calculations.
+	 * @param MeshHorizontalRange Arc size in degrees used for visible tiles calculations, specific to the sphere.
 	 */
-	void AddObject(AActor* InActor, float Width, float LODBias, EMediaTextureVisibleMipsTiles MeshType = EMediaTextureVisibleMipsTiles::None);
+	void AddObject(AActor* InActor, float LODBias = 0.0f, EMediaTextureVisibleMipsTiles MeshType = EMediaTextureVisibleMipsTiles::None, float MeshHorizontalRange = 360.0f);
 
 	/**
 	 * This object is no longer using our img sequence.
diff --git a/Engine/Plugins/Media/ImgMedia/Source/ImgMedia/Private/Assets/ImgMediaSource.cpp b/Engine/Plugins/Media/ImgMedia/Source/ImgMedia/Private/Assets/ImgMediaSource.cpp
--- a/Engine/Plugins/Media/ImgMedia/Source/ImgMedia/Private/Assets/ImgMediaSource.cpp
+++ b/Engine/Plugins/Media/ImgMedia/Source/ImgMedia/Private/Assets/ImgMediaSource.cpp
@@ -57,9 +57,9 @@ void UImgMediaSource::SetSequencePath(const FString& Path)
 }
 
 
-void UImgMediaSource::AddTargetObject(AActor* InActor, float Width)
+void UImgMediaSource::AddTargetObject(AActor* InActor)
 {
-	MipMapInfo->AddObject(InActor, Width, 0.0f);
+	MipMapInfo->AddObject(InActor);
 }
 
 
diff --git a/Engine/Plugins/Media/ImgMedia/Source/ImgMedia/Public/ImgMediaSource.h b/Engine/Plugins/Media/ImgMedia/Source/ImgMedia/Public/ImgMediaSource.h
--- a/Engine/Plugins/Media/ImgMedia/Source/ImgMedia/Public/ImgMediaSource.h
+++ b/Engine/Plugins/Media/ImgMedia/Source/ImgMedia/Public/ImgMediaSource.h
@@ -110,10 +110,9 @@ public:
 	 * This object is using our img sequence.
 	 *
 	 * @param InActor Object using our img sequence.
-	 * @param Width Width of the object. If < 0, then get the width automatically.
 	 */
 	UFUNCTION(BlueprintCallable, Category = "ImgMedia|ImgMediaSource")
-	void AddTargetObject(AActor* InActor, float Width = -1.0f);
+	void AddTargetObject(AActor* InActor);
 
 	/**
 	 * This object is no longer using our img sequence.
diff --git a/Engine/Plugins/Media/MediaPlate/Source/MediaPlate/Private/MediaPlateComponent.cpp b/Engine/Plugins/Media/MediaPlate/Source/MediaPlate/Private/MediaPlateComponent.cpp
--- a/Engine/Plugins/Media/MediaPlate/Source/MediaPlate/Private/MediaPlateComponent.cpp
+++ b/Engine/Plugins/Media/MediaPlate/Source/MediaPlate/Private/MediaPlateComponent.cpp
@@ -127,6 +127,7 @@ void UMediaPlateComponent::RegisterWithMediaTextureTracker()
 	MediaTextureTrackerObject->Object = GetOwner();
 	MediaTextureTrackerObject->MipMapLODBias = 0.0f;
 	MediaTextureTrackerObject->VisibleMipsTilesCalculations = VisibleMipsTilesCalculations;
+	MediaTextureTrackerObject->MeshHorizontalRange = MeshHorizontalRange;
 
 	// Add our texture.
 	TObjectPtr<UMediaTexture> MediaTexture = GetMediaTexture();
diff --git a/Engine/Plugins/Media/MediaPlate/Source/MediaPlate/Public/MediaPlateComponent.h b/Engine/Plugins/Media/MediaPlate/Source/MediaPlate/Public/MediaPlateComponent.h
--- a/Engine/Plugins/Media/MediaPlate/Source/MediaPlate/Public/MediaPlateComponent.h
+++ b/Engine/Plugins/Media/MediaPlate/Source/MediaPlate/Public/MediaPlateComponent.h
@@ -97,11 +97,8 @@ public:
 	UPROPERTY(EditAnywhere, AdvancedDisplay, Category = "MediaPlate")
 	EMediaTextureVisibleMipsTiles VisibleMipsTilesCalculations;
 
-	/**
-	 * Horizontal size of the mesh used for visible mips and tiles calculations.
-	 * For a sphere, this is the horizontal arc size in degrees.
-	 */
-	UPROPERTY(EditAnywhere, AdvancedDisplay, Category = "MediaPlate")
+	/** Horizontal arc size in degrees used for visible mips and tiles calculations, specific to the sphere. */
+	UPROPERTY(BlueprintReadOnly, Category = "MediaPlate")
 	float MeshHorizontalRange = 360.0f;
 
 	/**
diff --git a/Engine/Source/Runtime/MediaAssets/Public/MediaTextureTracker.h b/Engine/Source/Runtime/MediaAssets/Public/MediaTextureTracker.h
--- a/Engine/Source/Runtime/MediaAssets/Public/MediaTextureTracker.h
+++ b/Engine/Source/Runtime/MediaAssets/Public/MediaTextureTracker.h
@@ -30,6 +30,8 @@ struct FMediaTextureTrackerObject
 	float MipMapLODBias;
 	/** Specify type of mesh used for visible mips and tiles calculations. */
 	EMediaTextureVisibleMipsTiles VisibleMipsTilesCalculations;
+	/** Horizontal arc size in degrees used for visible tiles calculations, specific to the sphere.*/
+	float MeshHorizontalRange;
 };
 
 /**
