ImgMedia: Support mip level calculations for the generated spherical meshes.

diff --git a/Engine/Plugins/Media/ImgMedia/Source/ImgMedia/Private/Assets/ImgMediaMipMapInfo.cpp b/Engine/Plugins/Media/ImgMedia/Source/ImgMedia/Private/Assets/ImgMediaMipMapInfo.cpp
--- a/Engine/Plugins/Media/ImgMedia/Source/ImgMedia/Private/Assets/ImgMediaMipMapInfo.cpp
+++ b/Engine/Plugins/Media/ImgMedia/Source/ImgMedia/Private/Assets/ImgMediaMipMapInfo.cpp
@@ -289,14 +289,14 @@ namespace {
 	}
 
 	// Approximates hardware mip level selection.
-	bool CalculateMipLevel(const FImgMediaViewInfo& ViewInfo, const FVector& TexelWS, const FVector& OffsetXWS, const FVector& OffsetYWS, float& OutMipLevel)
+	bool CalculateMipLevel(const FImgMediaViewInfo& ViewInfo, const FVector& TexelWS, const FVector& TexelOffXWS, const FVector& TexelOffYWS, float& OutMipLevel)
 	{
 		FVector2D TexelScreenSpace[3];
 
 		bool bValid = true;
 		bValid &= ProjectWorldToScreenFast(TexelWS, ViewInfo.ViewportRect, ViewInfo.ViewProjectionMatrix, TexelScreenSpace[0]);
-		bValid &= ProjectWorldToScreenFast(TexelWS + OffsetXWS, ViewInfo.ViewportRect, ViewInfo.ViewProjectionMatrix, TexelScreenSpace[1]);
-		bValid &= ProjectWorldToScreenFast(TexelWS + OffsetYWS, ViewInfo.ViewportRect, ViewInfo.ViewProjectionMatrix, TexelScreenSpace[2]);
+		bValid &= ProjectWorldToScreenFast(TexelOffXWS, ViewInfo.ViewportRect, ViewInfo.ViewProjectionMatrix, TexelScreenSpace[1]);
+		bValid &= ProjectWorldToScreenFast(TexelOffYWS, ViewInfo.ViewportRect, ViewInfo.ViewProjectionMatrix, TexelScreenSpace[2]);
 
 		if (bValid)
 		{
@@ -359,7 +359,7 @@ namespace {
 				FConvexVolume ViewFrustum;
 				GetViewFrustumBounds(ViewFrustum, ViewInfo.OverscanViewProjectionMatrix, false, false);
 
-				int32 MaxLevel = InSequenceInfo.NumMipLevels - 1;
+				const int32 MaxLevel = InSequenceInfo.NumMipLevels - 1;
 				int MipLevelDiv = 1 << MaxLevel;
 
 				FIntPoint CurrentNumTiles;
@@ -432,7 +432,7 @@ namespace {
 									float CornerStepY = TileCornerY / (float)CurrentNumTiles.Y;
 									FVector CornersWS = PlaneCornerWS + (DirXWS * CornerStepX + DirYWS * CornerStepY);
 
-									if (CalculateMipLevel(ViewInfo, CornersWS, TexelOffsetXWS, TexelOffsetYWS, CalculatedLevel))
+									if (CalculateMipLevel(ViewInfo, CornersWS, CornersWS + TexelOffsetXWS, CornersWS + TexelOffsetYWS, CalculatedLevel))
 									{
 										CalculatedLevel += LODBias + ViewInfo.MaterialTextureMipBias;
 
@@ -550,7 +550,26 @@ namespace {
 				return;
 			}
 
+			const FTransform MeshTransform = Mesh->GetComponentTransform();
 			const float DefaultSphereRadius = 50.0f;
+			const int32 MaxLevel = InSequenceInfo.NumMipLevels - 1;
+
+			// Include all tiles containted in the visible UV region
+			int32 NumX = InSequenceInfo.NumTiles.X;
+			int32 NumY = InSequenceInfo.NumTiles.Y;
+			float PixelDimX = 1.0f / InSequenceInfo.Dim.X;
+			float PixelDimY = 1.0f / InSequenceInfo.Dim.Y;
+
+			FVector ApproxTileSizeWS = MeshTransform.GetScale3D() * (UE_TWO_PI * DefaultSphereRadius) / FMath::Max((float)NumX, (float)NumY);
+			float ApproxTileRadiusInWS = 0.5f * (float)FMath::Sqrt(2 * FMath::Square(ApproxTileSizeWS.GetAbsMax()));
+
+			auto TransformSphericalUVsToLocationWS = [this, &MeshTransform, DefaultSphereRadius](const FVector2D& InUV) -> FVector
+			{
+				// Convert from latlong UV to spherical coordinates
+				FVector2D TileCornerSpherical = FVector2D(UE_PI * InUV.Y, FMath::DegreesToRadians(MeshHorizontalRange) * (1.0f - InUV.X));
+				FVector CornersWS = TileCornerSpherical.SphericalToUnitCartesian() * DefaultSphereRadius;
+				return MeshTransform.TransformPosition(CornersWS);
+			};
 
 			for (const FImgMediaViewInfo& ViewInfo : InViewInfos)
 			{
@@ -562,71 +581,60 @@ namespace {
 				// Analytical derivation of visible tiles from the view frustum, given a sphere presumed to be infinitely large
 				FConvexVolume ViewFrustum;
 				GetViewFrustumBounds(ViewFrustum, ViewInfo.OverscanViewProjectionMatrix, false, false);
-				
-				// Include all tiles containted in the visible UV region
-				int32 NumX = InSequenceInfo.NumTiles.X;
-				int32 NumY = InSequenceInfo.NumTiles.Y;
+
 				for (int32 TileY = 0; TileY < NumY; ++TileY)
 				{
 					for (int32 TileX = 0; TileX < NumX; ++TileX)
 					{
-						FVector2D TileCornerUV = FVector2D((float)TileX / NumX, (float)TileY / NumY);
-						
-						// Convert from latlong UV to spherical coordinates
-						FVector2D TileCornerSpherical = FVector2D(UE_PI * TileCornerUV.Y, FMath::DegreesToRadians(MeshHorizontalRange) * (1.0f - TileCornerUV.X));
-
-						FVector TileCorner = TileCornerSpherical.SphericalToUnitCartesian() * DefaultSphereRadius;
-						TileCorner = Mesh->GetComponentTransform().TransformPosition(TileCorner);
+						FVector2D TileCenterUV = FVector2D(((float)TileX + 0.5f ) / NumX, ((float)TileY + 0.5f) / NumY);
+						FVector TileCenterWS = TransformSphericalUVsToLocationWS(TileCenterUV);
 
 						// For each tile corner, we include all adjacent tiles
-						if (ViewFrustum.IntersectPoint(TileCorner))
+						if (ViewFrustum.IntersectSphere(TileCenterWS, ApproxTileRadiusInWS))
 						{
-							if (!VisibleTiles.Contains(0))
-							{
-								VisibleTiles.Emplace(0, FImgMediaTileSelection(NumX, NumY));
-							}
+							float CalculatedLevel;
+							FIntVector2 MipLevelRange;
 
-							int32 AdjacentX = TileX > 0 ? TileX - 1 : NumX - 1;
-							int32 AdjacentY = TileY > 0 ? TileY - 1 : NumY - 1;
-							VisibleTiles[0].SetVisible(TileX, TileY);
-							VisibleTiles[0].SetVisible(AdjacentX, TileY);
-							VisibleTiles[0].SetVisible(TileX, AdjacentY);
-							VisibleTiles[0].SetVisible(AdjacentX, AdjacentY);
-						}
+							FVector TexelOffXWS = TransformSphericalUVsToLocationWS(TileCenterUV + FVector2D(PixelDimX, 0));
+							FVector TexelOffYWS = TransformSphericalUVsToLocationWS(TileCenterUV + FVector2D(0, PixelDimY));
 
-#if false
-#if WITH_EDITOR
-						// Enable this to draw a sphere where each tile is.
-						Async(EAsyncExecution::TaskGraphMainThread, [TileCorner]()
+							if (CalculateMipLevel(ViewInfo, TileCenterWS, TexelOffXWS, TexelOffYWS, CalculatedLevel))
 							{
-								UWorld* World = GEditor->GetEditorWorldContext().World();
-								DrawDebugPoint(World, TileCorner, 5.0f, FColor::Red, false, 0.05f);
-							});
-#endif // WITH_EDITOR
-#endif // false
-					}
-				}
+								CalculatedLevel += LODBias + ViewInfo.MaterialTextureMipBias;
 
-				if (VisibleTiles.Contains(0))
-				{
-					FIntPoint BaseDim = VisibleTiles[0].GetDimensions();
-					TArray<FIntPoint> BaseVisibleCoordinates = VisibleTiles[0].GetVisibleCoordinates();
+								MipLevelRange[0] = FMath::Clamp((int32)CalculatedLevel, 0, MaxLevel);
+								MipLevelRange[1] = FMath::CeilToInt32(CalculatedLevel);
 
-					// Include tiles visible at the base level in higher mip levels.
-					for (int32 Level = 1; Level < InSequenceInfo.NumMipLevels; ++Level)
-					{
-						const int32 MipLevelDiv = 1 << Level;
+								// As a mitigation for discontinuities at the poles, we artifically increase the max calculated level.
+								// (Note: Using an icosphere would avoid this issue but conflict with the partial sphere feature.)
+								if (TileY == 0 || TileY == NumY - 1)
+								{
+									MipLevelRange[1]++;
+								}
 
-						if (!VisibleTiles.Contains(Level))
-						{
-							NumX = FMath::CeilToInt((float)BaseDim.X / MipLevelDiv);
-							NumY = FMath::CeilToInt((float)BaseDim.Y / MipLevelDiv);
-							VisibleTiles.Emplace(Level, FImgMediaTileSelection(NumX, NumY));
-						}
+								MipLevelRange[1] = FMath::Clamp(MipLevelRange[1], 0, MaxLevel);
 
-						for (const FIntPoint& Coord : BaseVisibleCoordinates)
-						{
-							VisibleTiles[Level].SetVisible(Coord.X / MipLevelDiv, Coord.Y / MipLevelDiv);
+								for (int32 Level = MipLevelRange[0]; Level <= MipLevelRange[1]; ++Level)
+								{
+									if (!VisibleTiles.Contains(Level))
+									{
+										VisibleTiles.Emplace(Level, FImgMediaTileSelection::CreateForTargetMipLevel(InSequenceInfo.NumTiles.X, InSequenceInfo.NumTiles.Y, Level, false));
+									}
+
+									const int MipLevelDiv = 1 << Level;
+									VisibleTiles[Level].SetVisible(TileX / MipLevelDiv, TileY / MipLevelDiv);
+								}
+							}
+#if false
+#if WITH_EDITOR
+							// Enable this to draw a sphere where each tile is.
+							Async(EAsyncExecution::TaskGraphMainThread, [TileCenterWS, ApproxTileRadiusInWS]()
+								{
+									UWorld* World = GEditor->GetEditorWorldContext().World();
+									DrawDebugSphere(World, TileCenterWS, ApproxTileRadiusInWS, 8, FColor::Red, false, 0.05f);
+								});
+#endif // WITH_EDITOR
+#endif // false
 						}
 					}
 				}
diff --git a/Engine/Plugins/Media/MediaPlate/Source/MediaPlate/Private/MediaPlateComponent.cpp b/Engine/Plugins/Media/MediaPlate/Source/MediaPlate/Private/MediaPlateComponent.cpp
--- a/Engine/Plugins/Media/MediaPlate/Source/MediaPlate/Private/MediaPlateComponent.cpp
+++ b/Engine/Plugins/Media/MediaPlate/Source/MediaPlate/Private/MediaPlateComponent.cpp
@@ -120,6 +120,16 @@ void UMediaPlateComponent::Stop()
 	}
 }
 
+void UMediaPlateComponent::SetMeshHorizontalRange(float InMeshHorizontalRange)
+{
+	MeshHorizontalRange = InMeshHorizontalRange;
+
+	if (MediaTextureTrackerObject != nullptr)
+	{
+		MediaTextureTrackerObject->MeshHorizontalRange = MeshHorizontalRange;
+	}
+}
+
 void UMediaPlateComponent::RegisterWithMediaTextureTracker()
 {
 	// Set up object.
diff --git a/Engine/Plugins/Media/MediaPlate/Source/MediaPlate/Public/MediaPlateComponent.h b/Engine/Plugins/Media/MediaPlate/Source/MediaPlate/Public/MediaPlateComponent.h
--- a/Engine/Plugins/Media/MediaPlate/Source/MediaPlate/Public/MediaPlateComponent.h
+++ b/Engine/Plugins/Media/MediaPlate/Source/MediaPlate/Public/MediaPlateComponent.h
@@ -97,9 +97,13 @@ public:
 	UPROPERTY(EditAnywhere, AdvancedDisplay, Category = "MediaPlate")
 	EMediaTextureVisibleMipsTiles VisibleMipsTilesCalculations;
 
-	/** Horizontal arc size in degrees used for visible mips and tiles calculations, specific to the sphere. */
-	UPROPERTY(BlueprintReadOnly, Category = "MediaPlate")
-	float MeshHorizontalRange = 360.0f;
+	/** Set the horizontal arc size in degrees used for visible mips and tiles calculations, specific to the sphere. */
+	UFUNCTION(BlueprintCallable, Category = "Media|MediaPlate")
+	void SetMeshHorizontalRange(float InMeshHorizontalRange);
+
+	/** Return the horizontal arc size in degrees used for visible mips and tiles calculations, specific to the sphere. */
+	UFUNCTION(BlueprintCallable, Category = "Media|MediaPlate")
+	float GetMeshHorizontalRange() const { return MeshHorizontalRange; }
 
 	/**
 	 * Adds our media texture to the media texture tracker.
@@ -116,6 +120,9 @@ public:
 #endif
 
 private:
+	UPROPERTY()
+	float MeshHorizontalRange = 360.0f;
+
 	/** Name for our media component. */
 	static FLazyName MediaComponentName;
 	/** Name for our playlist. */
diff --git a/Engine/Plugins/Media/MediaPlate/Source/MediaPlateEditor/Private/MediaPlateCustomization.cpp b/Engine/Plugins/Media/MediaPlate/Source/MediaPlateEditor/Private/MediaPlateCustomization.cpp
--- a/Engine/Plugins/Media/MediaPlate/Source/MediaPlateEditor/Private/MediaPlateCustomization.cpp
+++ b/Engine/Plugins/Media/MediaPlate/Source/MediaPlateEditor/Private/MediaPlateCustomization.cpp
@@ -581,9 +581,9 @@ void FMediaPlateCustomization::SetMeshHorizontalRange(float HorizontalRange)
 		UMediaPlateComponent* MediaPlate = MediaPlatePtr.Get();
 		if (MediaPlate != nullptr)
 		{
-			if (MediaPlate->MeshHorizontalRange != HorizontalRange)
+			if (MediaPlate->GetMeshHorizontalRange() != HorizontalRange)
 			{
-				MediaPlate->MeshHorizontalRange = HorizontalRange;
+				MediaPlate->SetMeshHorizontalRange(HorizontalRange);
 				SetSphereMesh(MediaPlate);
 			}
 		}
@@ -598,7 +598,7 @@ TOptional<float> FMediaPlateCustomization::GetMeshHorizontalRange() const
 		UMediaPlateComponent* MediaPlate = MediaPlatePtr.Get();
 		if (MediaPlate != nullptr)
 		{
-			return MediaPlate->MeshHorizontalRange;
+			return MediaPlate->GetMeshHorizontalRange();
 		}
 	}
 
diff --git a/Engine/Plugins/Media/MediaPlate/Source/MediaPlateEditor/Private/MediaPlateCustomizationMesh.cpp b/Engine/Plugins/Media/MediaPlate/Source/MediaPlateEditor/Private/MediaPlateCustomizationMesh.cpp
--- a/Engine/Plugins/Media/MediaPlate/Source/MediaPlateEditor/Private/MediaPlateCustomizationMesh.cpp
+++ b/Engine/Plugins/Media/MediaPlate/Source/MediaPlateEditor/Private/MediaPlateCustomizationMesh.cpp
@@ -139,7 +139,7 @@ void FMediaPlateCustomizationMesh::GenerateSphereMesh(FDynamicMesh3* OutMesh,
 {
 	FMediaPlateSphereGenerator SphereGen;
 	SphereGen.Radius = 50.0f;
-	SphereGen.ThetaRange = FMath::DegreesToRadians(MediaPlate->MeshHorizontalRange);
+	SphereGen.ThetaRange = FMath::DegreesToRadians(MediaPlate->GetMeshHorizontalRange());
 	SphereGen.NumTheta = 64;
 	SphereGen.NumPhi = 64;
 	SphereGen.bPolygroupPerQuad = false;
@@ -183,7 +183,7 @@ UStaticMesh* FMediaPlateCustomizationMesh::CreateStaticMeshAsset(FDynamicMesh3*
 FString FMediaPlateCustomizationMesh::GetAssetPath(UMediaPlateComponent* MediaPlate)
 {
 	// Add the horizontal range.
-	FString ID = FString::SanitizeFloat(MediaPlate->MeshHorizontalRange);
+	FString ID = FString::SanitizeFloat(MediaPlate->GetMeshHorizontalRange());
 	ID.ReplaceCharInline(TCHAR('.'), TCHAR('_'));
 
 	FString AssetPath = FString::Printf(TEXT("/Game/_MediaPlate/Sphere_%s"), *ID);
