ImgMedia: Fix crash on process images if exr sequence folders contains (mac-style) hidden files.

#rnx

diff --git a/Engine/Plugins/Media/ImgMedia/Source/ImgMediaEditor/Private/Widgets/SImgMediaProcessImages.cpp b/Engine/Plugins/Media/ImgMedia/Source/ImgMediaEditor/Private/Widgets/SImgMediaProcessImages.cpp
--- a/Engine/Plugins/Media/ImgMedia/Source/ImgMediaEditor/Private/Widgets/SImgMediaProcessImages.cpp
+++ b/Engine/Plugins/Media/ImgMedia/Source/ImgMediaEditor/Private/Widgets/SImgMediaProcessImages.cpp
@@ -213,6 +213,7 @@ void SImgMediaProcessImages::ProcessAllImages()
 	
 	TArray<FString> FoundFiles;
 	IFileManager::Get().FindFiles(FoundFiles, *SequencePath, TEXT("*"));
+	FoundFiles.RemoveAll([](const FString& FileName) { return FileName.StartsWith(TEXT(".")); });
 	FoundFiles.Sort();
 	UE_LOG(LogImgMediaEditor, Warning, TEXT("Found %i image files in %s to import."), FoundFiles.Num(), *SequencePath);
 	if (FoundFiles.Num() == 0)
