ImgMedia Exr: Shader Resource View is created once for Unswizzled buffer instead of every frame. A tiny perf improvement.

diff --git a/Engine/Plugins/Media/ImgMedia/Source/ImgMedia/Private/Readers/ExrImgMediaReaderGpu.cpp b/Engine/Plugins/Media/ImgMedia/Source/ImgMedia/Private/Readers/ExrImgMediaReaderGpu.cpp
--- a/Engine/Plugins/Media/ImgMedia/Source/ImgMedia/Private/Readers/ExrImgMediaReaderGpu.cpp
+++ b/Engine/Plugins/Media/ImgMedia/Source/ImgMedia/Private/Readers/ExrImgMediaReaderGpu.cpp
@@ -469,7 +469,7 @@ void FExrImgMediaReaderGpu::CreateSampleConverterCallback
 					PermutationVector.Set<FExrSwizzlePS::FPartialTiles>(true);
 				}
 
-				Parameters.UnswizzledBuffer = RHICreateShaderResourceView(BufferData->BufferRef);
+				Parameters.UnswizzledBuffer = BufferData->ShaderResrouceView;
 
 				FGlobalShaderMap* ShaderMap = GetGlobalShaderMap(GMaxRHIFeatureLevel);
 
@@ -543,6 +543,7 @@ FStructuredBufferPoolItemSharedPtr FExrImgMediaReaderGpu::AllocateGpuBufferFromP
 				FRHIResourceCreateInfo CreateInfo(TEXT("FExrImgMediaReaderGpu"));
 				AllocatedBuffer->BufferRef = RHICreateStructuredBuffer(sizeof(uint16) * 2., AllocSize, BUF_ShaderResource | BUF_Dynamic | BUF_FastVRAM, CreateInfo);
 				AllocatedBuffer->MappedBuffer = RHILockBuffer(AllocatedBuffer->BufferRef, 0, AllocSize, RLM_WriteOnly);
+				AllocatedBuffer->ShaderResrouceView = RHICreateShaderResourceView(AllocatedBuffer->BufferRef);
 				AllocatedBuffer->Fence = RHICreateGPUFence(TEXT("BufferNoLongerInUseFence"));
 				if (bWait)
 				{
