ImgMedia Exr: Shader Resource View is created once for Unswizzled buffer instead of every frame. A tiny perf improvement.

diff --git a/Engine/Plugins/Media/ImgMedia/Source/ImgMedia/Private/Readers/ExrImgMediaReaderGpu.h b/Engine/Plugins/Media/ImgMedia/Source/ImgMedia/Private/Readers/ExrImgMediaReaderGpu.h
--- a/Engine/Plugins/Media/ImgMedia/Source/ImgMedia/Private/Readers/ExrImgMediaReaderGpu.h
+++ b/Engine/Plugins/Media/ImgMedia/Source/ImgMedia/Private/Readers/ExrImgMediaReaderGpu.h
@@ -24,6 +24,11 @@ struct FStructuredBufferPoolItem
 	*/
 	void* MappedBuffer;
 
+	/** 
+	* Resource View used by swizzling shader.
+	*/
+	FShaderResourceViewRHIRef ShaderResrouceView;
+
 	/** 
 	* A Gpu fence that identifies if this pool item is available for use again.
 	*/
