ImgMedia: Support {engine_dir} and {project_dir} tokens in sequence path.

This allows setting up image sequences in Engine plugins, because their path can be relative to the Engine folder, and before the only relative paths supported were with respect to the project path.

e.g.

{engine_dir}/Plugins/VirtualProduction/ExrLibrary/Content/Movies

* Deprecated IsPathRelativeToProjectRoot property. This property was used to determine if the path was relative to project or to project/Content. The meaning was sometimes confused in that false was thought to mean absolute ("not relative to project"), as opposed to still relative but to the project's Content folder. With this change, the user can still write paths relative to either the project or the content folder, and which one is decided by existence. Deprecation path adjusts the realtive path of existing image sequence media sources.

* Improved sanitation logic. Does not depend on "/." being there to detect relative paths. Trims spaces and double quotes - which often come when the user uses "Copy As Path" in Windows.

diff --git a/Engine/Plugins/Media/ImgMedia/Source/ImgMedia/Private/Assets/ImgMediaSource.cpp b/Engine/Plugins/Media/ImgMedia/Source/ImgMedia/Private/Assets/ImgMediaSource.cpp
--- a/Engine/Plugins/Media/ImgMedia/Source/ImgMedia/Private/Assets/ImgMediaSource.cpp
+++ b/Engine/Plugins/Media/ImgMedia/Source/ImgMedia/Private/Assets/ImgMediaSource.cpp
@@ -15,7 +15,7 @@
  *****************************************************************************/
 
 UImgMediaSource::UImgMediaSource()
-	: IsPathRelativeToProjectRoot(false)
+	: IsPathRelativeToProjectRoot_DEPRECATED(true)
 	, FrameRateOverride(0, 0)
 	, bFillGapsInSequence(true)
 	, MipMapInfo(MakeShared<FImgMediaMipMapInfo, ESPMode::ThreadSafe>())
@@ -31,12 +31,15 @@ void UImgMediaSource::GetProxies(TArray<FString>& OutProxies) const
 	IFileManager::Get().FindFiles(OutProxies, *FPaths::Combine(GetFullPath(), TEXT("*")), false, true);
 }
 
+const FString UImgMediaSource::GetSequencePath() const
+{
+	return ExpandSequencePathTokens(SequencePath.Path);
+}
 
 void UImgMediaSource::SetSequencePath(const FString& Path)
 {
 	const FString SanitizedPath = FPaths::GetPath(Path);
 
-	IsPathRelativeToProjectRoot = false;
 	if (SanitizedPath.IsEmpty() || SanitizedPath.StartsWith(TEXT(".")))
 	{
 		SequencePath.Path = SanitizedPath;
@@ -44,11 +47,11 @@ void UImgMediaSource::SetSequencePath(const FString& Path)
 	else
 	{
 		FString FullPath = FPaths::ConvertRelativePathToFull(SanitizedPath);
-		const FString FullGameContentDir = FPaths::ConvertRelativePathToFull(FPaths::ProjectContentDir());
+		const FString RelativeDir = FPaths::ConvertRelativePathToFull(FPaths::ProjectDir());
 
-		if (FullPath.StartsWith(FullGameContentDir))
+		if (FullPath.StartsWith(RelativeDir))
 		{
-			FPaths::MakePathRelativeTo(FullPath, *FullGameContentDir);
+			FPaths::MakePathRelativeTo(FullPath, *RelativeDir);
 			FullPath = FString(TEXT("./")) + FullPath;
 		}
 
@@ -56,6 +59,55 @@ void UImgMediaSource::SetSequencePath(const FString& Path)
 	}
 }
 
+void UImgMediaSource::SetTokenizedSequencePath(const FString& Path)
+{
+	SequencePath.Path = SanitizeTokenizedSequencePath(Path);
+}
+
+FString UImgMediaSource::ExpandSequencePathTokens(const FString& InPath)
+{
+	return InPath
+		.Replace(TEXT("{engine_dir}"), *FPaths::ConvertRelativePathToFull(FPaths::EngineDir()))
+		.Replace(TEXT("{project_dir}"), *FPaths::ConvertRelativePathToFull(FPaths::ProjectDir()))
+		;
+}
+
+FString UImgMediaSource::SanitizeTokenizedSequencePath(const FString& InPath)
+{
+	FString SanitizedPickedPath = InPath.TrimStartAndEnd().Replace(TEXT("\""), TEXT(""));
+
+	const FString ProjectAbsolutePath = FPaths::ConvertRelativePathToFull(FPaths::ProjectDir());
+
+	// Replace supported tokens
+	FString ExpandedPath = UImgMediaSource::ExpandSequencePathTokens(SanitizedPickedPath);
+
+	// Relative paths are always w.r.t. the project root.
+	if (FPaths::IsRelative(ExpandedPath))
+	{
+		ExpandedPath = FPaths::Combine(ProjectAbsolutePath, SanitizedPickedPath);
+	}
+
+	// Chop trailing file path, in case the user picked a file instead of a folder
+	if (FPaths::FileExists(ExpandedPath))
+	{
+		ExpandedPath = FPaths::GetPath(ExpandedPath);
+		SanitizedPickedPath = FPaths::GetPath(SanitizedPickedPath);
+	}
+
+	// If the user picked the absolute path of a directory that is inside the project, use relative path.
+	// Unless the user has a token in the beginning.
+	if (!InPath.Len() || InPath[0] != '{') // '{' indicates that the path begins with a token
+	{
+		FString PathRelativeToProject;
+
+		if (IsPathUnderBasePath(ExpandedPath, ProjectAbsolutePath, PathRelativeToProject))
+		{
+			SanitizedPickedPath = PathRelativeToProject;
+		}
+	}
+
+	return SanitizedPickedPath;
+}
 
 void UImgMediaSource::AddTargetObject(AActor* InActor)
 {
@@ -170,24 +222,40 @@ void UImgMediaSource::PostEditChangeProperty(FPropertyChangedEvent& PropertyChan
 
 FString UImgMediaSource::GetFullPath() const
 {
-	if (!FPaths::IsRelative(SequencePath.Path))
+	const FString ExpandedSequencePath = GetSequencePath();
+
+	if (FPaths::IsRelative(ExpandedSequencePath))
 	{
-		return SequencePath.Path;
+		return FPaths::ConvertRelativePathToFull(FPaths::Combine(FPaths::ProjectDir(), ExpandedSequencePath));
 	}
+	else
+	{
+		return ExpandedSequencePath;
+	}
+}
 
-	if (SequencePath.Path.StartsWith(TEXT("./")))
+void UImgMediaSource::Serialize(FArchive& Ar)
+{
+	Super::Serialize(Ar);
+
+	if (Ar.IsLoading() && !IsPathRelativeToProjectRoot_DEPRECATED)
 	{
-		FString RelativeDir;
-		if (IsPathRelativeToProjectRoot)
-		{
-			RelativeDir = FPaths::ProjectDir();
-		}
-		else
+		// This is an object that was saved with the old value, so we need to convert the path accordingly
+
+		IsPathRelativeToProjectRoot_DEPRECATED = true;
+
+		if (FPaths::IsRelative(SequencePath.Path))
 		{
-			RelativeDir = FPaths::ProjectContentDir();
+			SequencePath.Path = FString::Printf(TEXT("Content/%s"), *SequencePath.Path);
 		}
-		return FPaths::ConvertRelativePathToFull(RelativeDir, SequencePath.Path.RightChop(2));
 	}
+}
+
+bool UImgMediaSource::IsPathUnderBasePath(const FString& InPath, const FString& InBasePath, FString& OutRelativePath)
+{
+	OutRelativePath = InPath;
 
-	return FPaths::ConvertRelativePathToFull(SequencePath.Path);
+	return 
+		FPaths::MakePathRelativeTo(OutRelativePath, *InBasePath) 
+		&& !OutRelativePath.StartsWith(TEXT(".."));
 }
diff --git a/Engine/Plugins/Media/ImgMedia/Source/ImgMedia/Public/ImgMediaSource.h b/Engine/Plugins/Media/ImgMedia/Source/ImgMedia/Public/ImgMediaSource.h
--- a/Engine/Plugins/Media/ImgMedia/Source/ImgMedia/Public/ImgMediaSource.h
+++ b/Engine/Plugins/Media/ImgMedia/Source/ImgMedia/Public/ImgMediaSource.h
@@ -51,8 +51,8 @@ public:
 public:
 
 	/** If true, then relative Sequence Paths are relative to the project root directory. If false, then relative to the Content directory. */
-	UPROPERTY(BlueprintReadOnly, EditAnywhere, Category=Sequence)
-	bool IsPathRelativeToProjectRoot;
+	UPROPERTY()
+	bool IsPathRelativeToProjectRoot_DEPRECATED;
 
 	/** Overrides the default frame rate stored in the image files (0/0 = do not override). */
 	UPROPERTY(BlueprintReadWrite, EditAnywhere, Category=Sequence, AdvancedDisplay)
@@ -86,16 +86,13 @@ public:
 	void GetProxies(TArray<FString>& OutProxies) const;
 
 	/**
-	 * Get the path to the image sequence directory to be played.
+	 * Get the path to the image sequence directory to be played. Supported tokens will be expanded.
 	 *
 	 * @return The file path.
 	 * @see SetSequencePath
 	 */
-	UFUNCTION(BlueprintCallable, Category="ImgMedia|ImgMediaSource")
-	const FString GetSequencePath() const
-	{
-		return SequencePath.Path;
-	}
+	UFUNCTION(BlueprintCallable, Category = "ImgMedia|ImgMediaSource")
+	const FString GetSequencePath() const;
 
 	/**
 	 * Set the path to the image sequence directory this source represents.
@@ -106,6 +103,15 @@ public:
 	UFUNCTION(BlueprintCallable, Category="ImgMedia|ImgMediaSource")
 	void SetSequencePath(const FString& Path);
 
+	/**
+	 * Set the path to the image sequence directory this source represents.
+	 *
+	 * @param Path The path to the desired image sequence directory. May contain supported tokens.
+	 * @see GetSequencePath
+	 */
+	UFUNCTION(BlueprintCallable, Category = "ImgMedia|ImgMediaSource")
+	void SetTokenizedSequencePath(const FString& Path);
+
 	/**
 	 * This object is using our img sequence.
 	 *
@@ -150,14 +156,35 @@ public:
 	virtual void PostEditChangeProperty(FPropertyChangedEvent& PropertyChangedEvent) override;
 #endif // WITH_EDITOR
 
-protected:
+public:
+
+	//~ UObject interface
+
+	virtual void Serialize(FArchive& Ar) override;
+
+public:
+
+	/* Returns the path after replacing the supported tokens */
+	static FString ExpandSequencePathTokens(const FString& InPath);
+
+	/* Returns a sanitized sequence path, but without expanding the tokens tokens */
+	static FString SanitizeTokenizedSequencePath(const FString& InPath);
+
+	/* Returns true if InPath is under InBasePath. If so, fills out OutRelativePath with the relative path between them */
+	static bool IsPathUnderBasePath(const FString& InPath, const FString& InBasePath, FString& OutRelativePath);
+
+public:
 
 	/** Get the full path to the image sequence. */
 	FString GetFullPath() const;
 
 protected:
 
-	/** The directory that contains the image sequence files. */
+	/** The directory that contains the image sequence files.
+	 *  
+	 * - Relative paths will be with respect to the current Project directory. 
+	 * - You may use {engine_dir} or {project_dir} tokens.
+	 */
 	UPROPERTY(BlueprintReadOnly, EditAnywhere, Category=Sequence, meta = (EditCondition = "false", EditConditionHides))
 	FDirectoryPath SequencePath;
 
diff --git a/Engine/Plugins/Media/ImgMedia/Source/ImgMediaEditor/Private/Customizations/ImgMediaSourceCustomization.cpp b/Engine/Plugins/Media/ImgMedia/Source/ImgMediaEditor/Private/Customizations/ImgMediaSourceCustomization.cpp
--- a/Engine/Plugins/Media/ImgMedia/Source/ImgMediaEditor/Private/Customizations/ImgMediaSourceCustomization.cpp
+++ b/Engine/Plugins/Media/ImgMedia/Source/ImgMediaEditor/Private/Customizations/ImgMediaSourceCustomization.cpp
@@ -76,7 +76,7 @@
 					.BrowseDirectory_Lambda([this]() -> FString
 					{
 						const FString SequencePath = GetSequencePath();
-						return !SequencePath.IsEmpty() ? SequencePath : (FPaths::ProjectContentDir() / TEXT("Movies"));
+						return SequencePath.IsEmpty() ? (FPaths::ProjectContentDir() / TEXT("Movies")) : SequencePath;
 					})
 					.FilePath_Lambda([this]() -> FString
 					{
@@ -115,28 +115,9 @@
 
 FString FImgMediaSourceCustomization::GetSequencePath() const
 {
-	FString FilePath = GetSequencePathFromChildProperty(PropertyHandle);
-
-	return FilePath;
+	return GetSequencePathFromChildProperty(PropertyHandle);
 }
 
-FString FImgMediaSourceCustomization::GetRelativePathRoot() const
-{
-	FString RelativeDir;
-	bool bIsPathRelativeToRoot = IsPathRelativeToRoot();
-
-	if (bIsPathRelativeToRoot)
-	{
-		RelativeDir = FPaths::ConvertRelativePathToFull(FPaths::ProjectDir());
-	}
-	else
-	{
-		RelativeDir = FPaths::ConvertRelativePathToFull(FPaths::ProjectContentDir());
-	}
-	return RelativeDir;
-}
-
-
 TSharedPtr<IPropertyHandle> FImgMediaSourceCustomization::GetSequencePathProperty(const TSharedPtr<IPropertyHandle>& InPropertyHandle)
 {
 	TSharedPtr<IPropertyHandle> SequencePathProperty;
@@ -166,84 +147,73 @@
 	return SequencePathPathProperty;
 }
 
-
-TSharedPtr<IPropertyHandle> FImgMediaSourceCustomization::GetPathRelativeToRootProperty() const
-{
-	TSharedPtr<IPropertyHandle> PathRelativeToRootProperty;
-
-	if ((PropertyHandle.IsValid()) && (PropertyHandle->IsValidHandle()))
-	{
-		TSharedPtr<IPropertyHandle> ParentHandle = PropertyHandle->GetParentHandle();
-		if (ParentHandle.IsValid())
-		{
-			PathRelativeToRootProperty = ParentHandle->GetChildHandle("IsPathRelativeToProjectRoot");
-		}
-	}
-
-	return PathRelativeToRootProperty;
-}
-
-
-bool FImgMediaSourceCustomization::IsPathRelativeToRoot() const
-{
-	TSharedPtr<IPropertyHandle> PathRelativeToRootProperty = GetPathRelativeToRootProperty();
-
-	bool bIsPathRelativeToRoot = false;
-	if (PathRelativeToRootProperty.IsValid())
-	{
-		if (PathRelativeToRootProperty->GetValue(bIsPathRelativeToRoot) != FPropertyAccess::Success)
-		{
-			UE_LOG(LogImgMediaEditor, Error, TEXT("FImgMediaSourceCustomization could not get IsPathRelativeToProjectRoot."));
-		}
-	}
-
-	return bIsPathRelativeToRoot;
-}
-
-void FImgMediaSourceCustomization::SetPathRelativeToRoot(bool bIsPathRelativeToRoot)
-{
-	TSharedPtr<IPropertyHandle> PathRelativeToRootProperty = GetPathRelativeToRootProperty();
-
-	if (PathRelativeToRootProperty.IsValid())
-	{
-		if (PathRelativeToRootProperty->SetValue(bIsPathRelativeToRoot) != FPropertyAccess::Success)
-		{
-			UE_LOG(LogImgMediaEditor, Error, TEXT("FImgMediaSourceCustomization could not set IsPathRelativeToProjectRoot."));
-		}
-	}
-}
-
 /* FImgMediaSourceCustomization callbacks
  *****************************************************************************/
 
 void FImgMediaSourceCustomization::HandleSequencePathPickerPathPicked(const FString& PickedPath)
 {
-	// fully expand path & strip optional file name
-	const FString OldRelativeDir = GetRelativePathRoot();
-	const FString FullPath = PickedPath.StartsWith(TEXT("./"))
-		? FPaths::ConvertRelativePathToFull(OldRelativeDir, PickedPath.RightChop(2))
-		: PickedPath;
+	FString SanitizedPickedPath = PickedPath.TrimStartAndEnd().Replace(TEXT("\""), TEXT(""));
 
-	const FString FullDir = FPaths::ConvertRelativePathToFull(FPaths::FileExists(FullPath) ? FPaths::GetPath(FullPath) : FullPath);
-	
+	// The user may have put a path that is relative to the project, or to the content, or to the running process.
+	{
+		const FString ExpandedPath = UImgMediaSource::ExpandSequencePathTokens(SanitizedPickedPath);
 
-	FString PickedDir = FullDir;
-	const FString RelativeDir = FPaths::ConvertRelativePathToFull(FPaths::ProjectDir());
-	bool bIsRelativePath = PickedDir.StartsWith(RelativeDir);
-	if (bIsRelativePath)
-	{
-		FPaths::MakePathRelativeTo(PickedDir, *RelativeDir);
-		PickedDir = FString(TEXT("./")) + PickedDir;
+		// We only need to interpret relative path (not absolute paths).
+		if (FPaths::IsRelative(ExpandedPath))
+		{
+			// See if the path is relative to project root or content.
+			// Pick the one that exists (if any)
+
+			FString SanitizedPickedPathReplacement = SanitizedPickedPath;
+
+			const TArray<TPair<FString, FString>> PossibleBasePaths
+			{
+				TPair<FString,FString>(FPaths::ProjectDir(), TEXT("")),
+				TPair<FString,FString>(FPaths::ProjectContentDir(), TEXT("Content")),
+			};
+
+			int32 ExistedCnt = 0;
+
+			for (const auto& Pair : PossibleBasePaths)
+			{
+				const FString CandidatePath = FPaths::Combine(Pair.Key, ExpandedPath);
+
+				if (FPaths::FileExists(CandidatePath) || FPaths::DirectoryExists(CandidatePath))
+				{
+					SanitizedPickedPathReplacement = FPaths::Combine(Pair.Value, SanitizedPickedPath);
+					ExistedCnt++;
+				}
+			}
+
+			if (ExistedCnt == 0)
+			{
+				// If the path didn't exist relative to project or content, see if it is relative to the current process,
+				// which is the kind of path returned by the file picker.
+				// In that case, change to full path.
+
+				const FString CandidatePathRelativeToEngine = FPaths::Combine(FPlatformProcess::BaseDir(), ExpandedPath);
+
+				if (FPaths::FileExists(CandidatePathRelativeToEngine) || FPaths::DirectoryExists(CandidatePathRelativeToEngine))
+				{
+					SanitizedPickedPath = FPaths::ConvertRelativePathToFull(CandidatePathRelativeToEngine);
+				}
+			}
+			else if (ExistedCnt == 1)
+			{
+				// Re-interpret the base relative if the path was relative to either project or content (leave unaltered when both exist)
+				SanitizedPickedPath = SanitizedPickedPathReplacement;
+			}
+		}
 	}
-	
-	// Update relative to root property.
-	SetPathRelativeToRoot(bIsRelativePath);
 
+	// Sanitize the path
+	SanitizedPickedPath = UImgMediaSource::SanitizeTokenizedSequencePath(SanitizedPickedPath);
+
 	// update property
 	TSharedPtr<IPropertyHandle> SequencePathPathProperty = GetSequencePathPathProperty(PropertyHandle);
 	if (SequencePathPathProperty.IsValid())
 	{
-		if (SequencePathPathProperty->SetValue(PickedDir) != FPropertyAccess::Success)
+		if (SequencePathPathProperty->SetValue(SanitizedPickedPath) != FPropertyAccess::Success)
 		{
 			UE_LOG(LogImgMediaEditor, Error, TEXT("FImgMediaSourceCustomization could not set SequencePath."));
 		}
@@ -253,18 +223,18 @@
 
 EVisibility FImgMediaSourceCustomization::HandleSequencePathWarningIconVisibility() const
 {
-	FString FilePath = GetSequencePath();
+	const FString FilePath = UImgMediaSource::ExpandSequencePathTokens(GetSequencePath());
 
 	if (FilePath.IsEmpty() || FilePath.Contains(TEXT("://")))
 	{
 		return EVisibility::Hidden;
 	}
 
-	const FString RelativeDir = GetRelativePathRoot();
 	const FString FullMoviesPath = FPaths::ConvertRelativePathToFull(FPaths::ProjectContentDir() / TEXT("Movies"));
-	const FString FullPath = FPaths::ConvertRelativePathToFull(FPaths::IsRelative(FilePath) ? RelativeDir / FilePath : FilePath);
+	const FString FullPath = FPaths::ConvertRelativePathToFull(FPaths::IsRelative(FilePath) ? FPaths::ProjectDir() / FilePath : FilePath);
+	FString RelativePath;
 
-	if (FullPath.StartsWith(FullMoviesPath))
+	if (UImgMediaSource::IsPathUnderBasePath(FullPath, FullMoviesPath, RelativePath))
 	{
 		if (FPaths::DirectoryExists(FullPath))
 		{
diff --git a/Engine/Plugins/Media/ImgMedia/Source/ImgMediaEditor/Private/Customizations/ImgMediaSourceCustomization.h b/Engine/Plugins/Media/ImgMedia/Source/ImgMediaEditor/Private/Customizations/ImgMediaSourceCustomization.h
--- a/Engine/Plugins/Media/ImgMedia/Source/ImgMediaEditor/Private/Customizations/ImgMediaSourceCustomization.h
+++ b/Engine/Plugins/Media/ImgMedia/Source/ImgMediaEditor/Private/Customizations/ImgMediaSourceCustomization.h
@@ -56,13 +56,6 @@
 	 */
 	FString GetSequencePath() const;
 
-	/**
-	 * Get the root path we are using for a relative Sequence Path.
-	 *
-	 * @return Root path.
-	 */
-	FString GetRelativePathRoot() const;
-
 private:
 
 	/** Callback for picking a path in the source directory picker. */
@@ -73,15 +66,10 @@
 
 	/** Returns the property for SequencePath. */
 	static TSharedPtr<IPropertyHandle> GetSequencePathProperty(const TSharedPtr<IPropertyHandle>& InPropertyHandle);
+	
 	/** Returns the property for SequencePath->Path. */
 	static TSharedPtr<IPropertyHandle> GetSequencePathPathProperty(const TSharedPtr<IPropertyHandle>& InPropertyHandle);
-	/** Returns the property for IsPathRelativeToProjectRoot. */
-	TSharedPtr<IPropertyHandle> GetPathRelativeToRootProperty() const;
-	/** Returns the value of IsPathRelativeToProjectRoot. */
-	bool IsPathRelativeToRoot() const;
-	/** Sets the value of IsPathRelativeToProjectRoot. */
-	void SetPathRelativeToRoot(bool bIsPathRelativeToRoot);
-
+	
 private:
 
 	/** Text block widget showing the found proxy directories. */
