ImgMedia: Fix race condition between mipmap info object destruction and the CalculateVisibleTiles routine.

#rnx

diff --git a/Engine/Plugins/Media/ImgMedia/Source/ImgMedia/Private/Assets/ImgMediaMipMapInfo.cpp b/Engine/Plugins/Media/ImgMedia/Source/ImgMedia/Private/Assets/ImgMediaMipMapInfo.cpp
--- a/Engine/Plugins/Media/ImgMedia/Source/ImgMedia/Private/Assets/ImgMediaMipMapInfo.cpp
+++ b/Engine/Plugins/Media/ImgMedia/Source/ImgMedia/Private/Assets/ImgMediaMipMapInfo.cpp
@@ -684,6 +684,8 @@ void FImgMediaMipMapInfo::RemoveObject(AActor* InActor)
 {
 	if (InActor != nullptr)
 	{
+		FScopeLock Lock(&InfoCriticalSection);
+
 		for (int Index = 0; Index < Objects.Num(); ++Index)
 		{
 			FImgMediaMipMapObjectInfo* Info = Objects[Index];
@@ -743,6 +745,8 @@ void FImgMediaMipMapInfo::RemoveObjectsUsingThisMediaTexture(UMediaTexture* InMe
 
 void FImgMediaMipMapInfo::ClearAllObjects()
 {
+	FScopeLock Lock(&InfoCriticalSection);
+
 	for (FImgMediaMipMapObjectInfo* Info : Objects)
 	{
 		delete Info;
