ImgMedia and FilePathPicker: Added option for the file picker to return an absolute path when the path comes from the system dialog.

This eliminates the ambiguity when the path is returned is relative, since for the image sequence path we support paths to be specified relative to the project (typed by the user) while the dialog returned paths relative to the process, but the callback for the textfield and the dialog are the same at the file picker layer. The issue ultimately happens because when ".." hits the root folder, it behaves as a no-op. So multiple relative paths can point to the same file for the same given starting directory.

The flag added to SFilePathPicker defaults to the previous behavior (which is to return paths relative to the process).

diff --git a/Engine/Plugins/Media/ImgMedia/Source/ImgMediaEditor/Private/Customizations/ImgMediaSourceCustomization.cpp b/Engine/Plugins/Media/ImgMedia/Source/ImgMediaEditor/Private/Customizations/ImgMediaSourceCustomization.cpp
--- a/Engine/Plugins/Media/ImgMedia/Source/ImgMediaEditor/Private/Customizations/ImgMediaSourceCustomization.cpp
+++ b/Engine/Plugins/Media/ImgMedia/Source/ImgMediaEditor/Private/Customizations/ImgMediaSourceCustomization.cpp
@@ -73,6 +73,7 @@ void FImgMediaSourceCustomization::CustomizeHeader(TSharedRef<IPropertyHandle> I
 					.BrowseButtonImage(FAppStyle::GetBrush("PropertyWindow.Button_Ellipsis"))
 					.BrowseButtonStyle(FAppStyle::Get(), "HoverHintOnly")
 					.BrowseButtonToolTip(LOCTEXT("SequencePathBrowseButtonToolTip", "Choose a file from this computer"))
+					.DialogReturnsFullPath(true)
 					.BrowseDirectory_Lambda([this]() -> FString
 					{
 						const FString SequencePath = GetSequencePath();
@@ -159,10 +160,11 @@ void FImgMediaSourceCustomization::HandleSequencePathPickerPathPicked(const FStr
 		const FString ExpandedPath = UImgMediaSource::ExpandSequencePathTokens(SanitizedPickedPath);
 
 		// We only need to interpret relative path (not absolute paths).
+		// Note: The path picker is set to only returns absolute paths.
 		if (FPaths::IsRelative(ExpandedPath))
 		{
 			// See if the path is relative to project root or content.
-			// Pick the one that exists (if any)
+			// Pick the one that exists
 
 			FString SanitizedPickedPathReplacement = SanitizedPickedPath;
 
@@ -185,22 +187,9 @@ void FImgMediaSourceCustomization::HandleSequencePathPickerPathPicked(const FStr
 				}
 			}
 
-			if (ExistedCnt == 0)
+			if (ExistedCnt == 1)
 			{
-				// If the path didn't exist relative to project or content, see if it is relative to the current process,
-				// which is the kind of path returned by the file picker.
-				// In that case, change to full path.
-
-				const FString CandidatePathRelativeToEngine = FPaths::Combine(FPlatformProcess::BaseDir(), ExpandedPath);
-
-				if (FPaths::FileExists(CandidatePathRelativeToEngine) || FPaths::DirectoryExists(CandidatePathRelativeToEngine))
-				{
-					SanitizedPickedPath = FPaths::ConvertRelativePathToFull(CandidatePathRelativeToEngine);
-				}
-			}
-			else if (ExistedCnt == 1)
-			{
-				// Re-interpret the base relative if the path was relative to either project or content (leave unaltered when both exist)
+				// Re-interpret the base relative if the path was relative to either project or content (otherwise leave unaltered)
 				SanitizedPickedPath = SanitizedPickedPathReplacement;
 			}
 		}
diff --git a/Engine/Source/Developer/DesktopWidgets/Private/Widgets/Input/SFilePathPicker.cpp b/Engine/Source/Developer/DesktopWidgets/Private/Widgets/Input/SFilePathPicker.cpp
--- a/Engine/Source/Developer/DesktopWidgets/Private/Widgets/Input/SFilePathPicker.cpp
+++ b/Engine/Source/Developer/DesktopWidgets/Private/Widgets/Input/SFilePathPicker.cpp
@@ -22,6 +22,7 @@ void SFilePathPicker::Construct( const FArguments& InArgs )
 	FilePath = InArgs._FilePath;
 	FileTypeFilter = InArgs._FileTypeFilter;
 	OnPathPicked = InArgs._OnPathPicked;
+	DialogReturnsFullPath = InArgs._DialogReturnsFullPath;
 
 	ChildSlot
 	[
@@ -91,7 +92,14 @@ FReply SFilePathPicker::HandleBrowseButtonClicked()
 
 	if (DesktopPlatform->OpenFileDialog(ParentWindowHandle, BrowseTitle.Get().ToString(), DefaultPath, TEXT(""), FileTypeFilter.Get(), EFileDialogFlags::None, OutFiles))
 	{
-		OnPathPicked.ExecuteIfBound(OutFiles[0]);
+		if (DialogReturnsFullPath.Get())
+		{
+			OnPathPicked.ExecuteIfBound(FPaths::ConvertRelativePathToFull(OutFiles[0]));
+		}
+		else
+		{
+			OnPathPicked.ExecuteIfBound(OutFiles[0]);
+		}
 	}
 
 	return FReply::Handled();
diff --git a/Engine/Source/Developer/DesktopWidgets/Public/Widgets/Input/SFilePathPicker.h b/Engine/Source/Developer/DesktopWidgets/Public/Widgets/Input/SFilePathPicker.h
--- a/Engine/Source/Developer/DesktopWidgets/Public/Widgets/Input/SFilePathPicker.h
+++ b/Engine/Source/Developer/DesktopWidgets/Public/Widgets/Input/SFilePathPicker.h
@@ -35,6 +35,7 @@ public:
 		, _FileTypeFilter(TEXT("All files (*.*)|*.*"))
 		, _Font()
 		, _IsReadOnly(false)
+		, _DialogReturnsFullPath(false)
 	{ }
 
 		/** Browse button image resource. */
@@ -64,6 +65,9 @@ public:
 		/** Whether the path text box can be modified by the user. */
 		SLATE_ATTRIBUTE(bool, IsReadOnly)
 
+		/** Whether the path returned by the dialog should be converted from relative to full */
+		SLATE_ATTRIBUTE(bool, DialogReturnsFullPath)
+
 		/** Called when a file path has been picked. */
 		SLATE_EVENT(FOnPathPicked, OnPathPicked)
 
@@ -104,6 +108,9 @@ private:
 	/** Holds the editable text box. */
 	TSharedPtr<SEditableTextBox> TextBox;
 
+	/** Holds the option for the dialog to return full path instead of relative. */
+	TAttribute<bool> DialogReturnsFullPath;
+
 private:
 
 	/** Holds a delegate that is executed when a file was picked. */
