ImgMedia: Prevent tile calculation when the mesh is hidden/invisible.

#rnx

diff --git a/Engine/Plugins/Media/ImgMedia/Source/ImgMedia/Private/Assets/ImgMediaMipMapInfo.cpp b/Engine/Plugins/Media/ImgMedia/Source/ImgMedia/Private/Assets/ImgMediaMipMapInfo.cpp
--- a/Engine/Plugins/Media/ImgMedia/Source/ImgMedia/Private/Assets/ImgMediaMipMapInfo.cpp
+++ b/Engine/Plugins/Media/ImgMedia/Source/ImgMedia/Private/Assets/ImgMediaMipMapInfo.cpp
@@ -264,6 +264,12 @@ UMeshComponent* FImgMediaMipMapObjectInfo::GetMeshComponent() const
 
 void FImgMediaMipMapObjectInfo::CalculateVisibleTiles(const TArray<FImgMediaViewInfo>& InViewInfos, const FSequenceInfo& InSequenceInfo, TMap<int32, FImgMediaTileSelection>& VisibleTiles) const
 {
+	UMeshComponent* Mesh = MeshComponent.Get();
+	if (Mesh == nullptr || !Mesh->ShouldRender())
+	{
+		return;
+	}
+
 	// We simply add fully visible regions for all mip levels
 	for (int32 MipLevel = 0; MipLevel < InSequenceInfo.NumMipLevels; ++MipLevel)
 	{
@@ -329,7 +335,7 @@ namespace {
 		void CalculateVisibleTiles(const TArray<FImgMediaViewInfo>& InViewInfos, const FSequenceInfo& InSequenceInfo, TMap<int32, FImgMediaTileSelection>& VisibleTiles) const override
 		{
 			UMeshComponent* Mesh = MeshComponent.Get();
-			if (Mesh == nullptr)
+			if (Mesh == nullptr || !Mesh->ShouldRender())
 			{
 				return;
 			}
@@ -545,7 +551,7 @@ namespace {
 		void CalculateVisibleTiles(const TArray<FImgMediaViewInfo>& InViewInfos, const FSequenceInfo& InSequenceInfo, TMap<int32, FImgMediaTileSelection>& VisibleTiles) const override
 		{
 			UMeshComponent* Mesh = MeshComponent.Get();
-			if (Mesh == nullptr)
+			if (Mesh == nullptr || !Mesh->ShouldRender())
 			{
 				return;
 			}
