ImgMedia ExrReaderGpu: Reader is now utilizing Copy thread on GPU. We utilize upload heap to read into and then issue a copy command into pure GPU buffer. This reduces the cost on 3D thread on GPU and increases overal performance by a significant margin.

diff --git a/Engine/Plugins/Media/ImgMedia/Source/ImgMedia/Private/Readers/ExrImgMediaReader.cpp b/Engine/Plugins/Media/ImgMedia/Source/ImgMedia/Private/Readers/ExrImgMediaReader.cpp
--- a/Engine/Plugins/Media/ImgMedia/Source/ImgMedia/Private/Readers/ExrImgMediaReader.cpp
+++ b/Engine/Plugins/Media/ImgMedia/Source/ImgMedia/Private/Readers/ExrImgMediaReader.cpp
@@ -55,19 +55,19 @@ bool FExrImgMediaReader::GetFrameInfo(const FString& ImagePath, FImgMediaFrameIn
 	return GetInfo(ImagePath, OutInfo);
 }
 
-FExrImgMediaReader::EReadResult FExrImgMediaReader::ReadTilesCustom
+FExrImgMediaReader::EReadResult FExrImgMediaReader::ReadTiles
 	( uint16* Buffer
 	, int64 BufferSize
 	, const FString& ImagePath
 	, int32 FrameId
 	, const TArray<FIntRect>& TileRegions
 	, TSharedPtr<FSampleConverterParameters> ConverterParams
-	, const int32 CurrentMipLevel)
+	, const int32 CurrentMipLevel
+	, TArray<UE::Math::TIntPoint<int64>>& OutBufferRegionsToCopy)
 {
 #if defined(PLATFORM_WINDOWS) && PLATFORM_WINDOWS
-	EReadResult bResult = Success;
+	EReadResult Result = Success;
 
-	int64 CurrentBufferPos = 0;
 	FExrReader ChunkReader;
 	int MipLevelDiv = 1 << CurrentMipLevel;
 
@@ -99,6 +99,7 @@ FExrImgMediaReader::EReadResult FExrImgMediaReader::ReadTilesCustom
 		}
 	}
 	{
+		int64 CurrentBufferPos = 0;
 		TRACE_CPUPROFILER_EVENT_SCOPE_TEXT(*FString::Printf(TEXT("FExrImgMediaReader_ReadTilesCustom_ReadTiles")));
 		for (const FIntRect& TileRegion : TileRegions)
 		{
@@ -109,8 +110,8 @@ FExrImgMediaReader::EReadResult FExrImgMediaReader::ReadTilesCustom
 					FScopeLock RegionScopeLock(&CanceledFramesCriticalSection);
 					if (CanceledFrames.Remove(FrameId) > 0)
 					{
-						UE_LOG(LogImgMedia, Warning, TEXT("Reader %p: Canceling Frame %i At tile row # %i"), this, FrameId, TileRow);
-						bResult = Cancelled;
+						UE_LOG(LogImgMedia, Verbose, TEXT("Reader %p: Canceling Frame %i At tile row # %i"), this, FrameId, TileRow);
+						Result = Cancelled;
 						break;
 					}
 				}
@@ -130,7 +131,7 @@ FExrImgMediaReader::EReadResult FExrImgMediaReader::ReadTilesCustom
 				if (!ChunkReader.GetByteOffsetForTile(StartTileIndex, MipLevel, ByteOffsetStart)
 					|| !ChunkReader.GetByteOffsetForTile(EndTileIndex, MipLevel, ByteOffsetEnd))
 				{
-					bResult = Fail;
+					Result = Fail;
 					break;
 				}
 
@@ -141,10 +142,12 @@ FExrImgMediaReader::EReadResult FExrImgMediaReader::ReadTilesCustom
 
 				if (!ChunkReader.ReadExrImageChunk(reinterpret_cast<char*>(Buffer) + CurrentBufferPos, ByteChunkToRead))
 				{
-					bResult = Fail;
+					Result = Fail;
 					break;
 				}
+				OutBufferRegionsToCopy.Add({ CurrentBufferPos, ByteChunkToRead });
 			}
+
 		}
 	}
 	{
@@ -154,7 +157,8 @@ FExrImgMediaReader::EReadResult FExrImgMediaReader::ReadTilesCustom
 		{
 			return Fail;
 		}
-		return bResult;
+
+		return Result;
 	}
 #else
 	return Fail;
@@ -264,7 +268,8 @@ bool FExrImgMediaReader::ReadFrame(int32 FrameId, const TMap<int32, FImgMediaTil
 					ConverterParams->NumMipLevels = Loader->GetNumMipLevels();
 					ConverterParams->bCustomExr = OutFrame->Info.FormatName == TEXT("EXR CUSTOM");
 
-					EReadResult ReadResult = ReadTilesCustom((uint16*)MipDataPtr, GetMipBufferTotalSize(Dim / MipLevelDiv), Image, FrameId, TileRegions, ConverterParams, CurrentMipLevel);
+					TArray<UE::Math::TIntPoint<int64>> OutBufferRegionsToCopy;
+					EReadResult ReadResult = ReadTiles((uint16*)MipDataPtr, GetMipBufferTotalSize(Dim / MipLevelDiv), Image, FrameId, TileRegions, ConverterParams, CurrentMipLevel, OutBufferRegionsToCopy);
 					if (ReadResult != Fail)
 					{
 						OutFrame->MipTilesPresent.Emplace(CurrentMipLevel, CurrentTileSelection);
diff --git a/Engine/Plugins/Media/ImgMedia/Source/ImgMedia/Private/Readers/ExrImgMediaReader.h b/Engine/Plugins/Media/ImgMedia/Source/ImgMedia/Private/Readers/ExrImgMediaReader.h
--- a/Engine/Plugins/Media/ImgMedia/Source/ImgMedia/Private/Readers/ExrImgMediaReader.h
+++ b/Engine/Plugins/Media/ImgMedia/Source/ImgMedia/Private/Readers/ExrImgMediaReader.h
@@ -57,6 +57,9 @@ protected:
 		/** General file information including its header. */
 		FImgMediaFrameInfo FrameInfo;
 
+		/** Frame Id. */
+		int32 FrameId;
+
 		/** Resolution of the highest quality mip. */
 		FIntPoint FullResolution;
 
@@ -67,7 +70,9 @@ protected:
 		TMap<int32, TArray<FIntRect>> Viewports;
 
 #if defined(PLATFORM_WINDOWS) && PLATFORM_WINDOWS
-		/** Contain information about individual tiles. Used to convert buffer data into a 2D Texture. */
+		/** Contain information about individual tiles. Used to convert buffer data into a 2D Texture. 
+		* The size of this array is the exact number of complete and partial tiles for each mip level.
+		*/
 		TArray<TArray<FExrReader::FTileDesc>> TileInfoPerMipLevel;
 #endif
 
@@ -96,8 +101,24 @@ protected:
 	/**
 	 * Reads tiles from exr files in tile rows based on Tile region. If frame is pending for cancelation
 	 * stops reading tiles at the current tile row.
+	 * 
+	 * @param Buffer					A buffer to read into from hard drive.
+	 * @param BufferSize				The total size of the buffer. Used to limit reader to buffer region.
+	 * @param ImagePath					Path to the file to read.
+	 * @param FrameId					Frame id used to determine if frame has been canceled.
+	 * @param TileRegions				Rectangular tile regions that we read from the hard drive line by line.
+	 * @param ConverterParams			Full information about the current frame.
+	 * @param OutBufferRegionsToCopy	This buffer is used to issue copy commands on GPU Copy thread.
 	*/
-	EReadResult ReadTilesCustom(uint16* Buffer, int64 BufferSize, const FString& ImagePath, int32 FrameId, const TArray<FIntRect>& TileRegions, TSharedPtr<FSampleConverterParameters> ConverterParams, const int32 CurrentMipLevel);
+	EReadResult ReadTiles
+		( uint16* Buffer
+		, int64 BufferSize
+		, const FString& ImagePath
+		, int32 FrameId
+		, const TArray<FIntRect>& TileRegions
+		, TSharedPtr<FSampleConverterParameters> ConverterParams
+		, const int32 CurrentMipLevel
+		, TArray<UE::Math::TIntPoint<int64>>& OutBufferRegionsToCopy);
 
 	/**
 	 * Sets parameters of our custom format images.
diff --git a/Engine/Plugins/Media/ImgMedia/Source/ImgMedia/Private/Readers/ExrImgMediaReaderGpu.cpp b/Engine/Plugins/Media/ImgMedia/Source/ImgMedia/Private/Readers/ExrImgMediaReaderGpu.cpp
--- a/Engine/Plugins/Media/ImgMedia/Source/ImgMedia/Private/Readers/ExrImgMediaReaderGpu.cpp
+++ b/Engine/Plugins/Media/ImgMedia/Source/ImgMedia/Private/Readers/ExrImgMediaReaderGpu.cpp
@@ -31,6 +31,11 @@
 DECLARE_GPU_STAT_NAMED(ExrImgMediaReaderGpu, TEXT("ExrImgMediaReaderGpu"));
 DECLARE_GPU_STAT_NAMED(ExrImgMediaReaderGpu_MipRender, TEXT("ExrImgMediaReaderGpu_MipRender"));
 
+static TAutoConsoleVariable<bool> CVarExrReaderUseUploadHeap(
+	TEXT("r.ExrReaderGPU.UseUploadHeap"),
+	true,
+	TEXT("Utilizes upload heap and copies raw exr buffer asynchronously.\n"));
+
 namespace {
 
 	/** This function is similar to DrawScreenPass in OpenColorIODisplayExtension.cpp except it is catered for Viewless texture rendering. */
@@ -92,9 +97,9 @@ FExrImgMediaReaderGpu::~FExrImgMediaReaderGpu()
 			for (FStructuredBufferPoolItem* MemoryPoolItem : AllValues)
 			{
 				// Check if fence has signaled.
-				check(!MemoryPoolItem->bWillBeSignaled || MemoryPoolItem->Fence->Poll());
+				check(!MemoryPoolItem->bWillBeSignaled || MemoryPoolItem->RenderFence->Poll());
 				{
-					RHIUnlockBuffer(MemoryPoolItem->BufferRef);
+					RHIUnlockBuffer(MemoryPoolItem->UploadBufferRef);
 					delete MemoryPoolItem;
 				}
 			}
@@ -141,6 +146,7 @@ bool FExrImgMediaReaderGpu::ReadFrame(int32 FrameId, const TMap<int32, FImgMedia
 
 	TSharedPtr<FSampleConverterParameters> ConverterParams = MakeShared<FSampleConverterParameters>();
 	ConverterParams->FullResolution = OutFrame->Info.Dim;
+	ConverterParams->FrameId = FrameId;
 
 	const FIntPoint TileDim = OutFrame->Info.TileDimensions;
 
@@ -172,7 +178,7 @@ bool FExrImgMediaReaderGpu::ReadFrame(int32 FrameId, const TMap<int32, FImgMedia
 
 			const int32 CurrentMipLevel = TilesPerMip.Key;
 
-			TRACE_CPUPROFILER_EVENT_SCOPE_TEXT(*FString::Printf(TEXT("FExrImgMediaReaderGpu_ReadMip %d"), CurrentMipLevel));
+			TRACE_CPUPROFILER_EVENT_SCOPE_TEXT(*FString::Printf(TEXT("ExrReaderGpu.ReadMip %d"), CurrentMipLevel));
 
 			const FImgMediaTileSelection& CurrentTileSelection = TilesPerMip.Value;
 
@@ -202,7 +208,7 @@ bool FExrImgMediaReaderGpu::ReadFrame(int32 FrameId, const TMap<int32, FImgMedia
 				SampleConverter->SetMipLevelBuffer(CurrentMipLevel, BufferData);
 			}
 
-			uint16* MipDataPtr = static_cast<uint16*>(BufferData->MappedBuffer);
+			uint16* MipDataPtr = static_cast<uint16*>(BufferData->UploadBufferMapped);
 
 			// Get highest resolution mip level path.
 			FString ImagePath = Loader->GetImagePath(FrameId, ConverterParams->bMipsInSeparateFiles ? CurrentMipLevel : 0);
@@ -211,13 +217,14 @@ bool FExrImgMediaReaderGpu::ReadFrame(int32 FrameId, const TMap<int32, FImgMedia
 
 			if (FPaths::FileExists(ImagePath))
 			{
+				TArray<UE::Math::TIntPoint<int64>> BufferRegionsToCopy;
 				// read frame data
 				if (bHasTiles || ConverterParams->bCustomExr)
 				{
 					TArray<FIntRect> TileRegionsToRead;
 					TArray<FIntRect> TileRegionsToRender;
 					{
-						TRACE_CPUPROFILER_EVENT_SCOPE_TEXT(*FString::Printf(TEXT("CalculateRegions %d"), CurrentMipLevel));
+						TRACE_CPUPROFILER_EVENT_SCOPE_TEXT(*FString::Printf(TEXT("ExrReaderGpu.CalculateRegions %d"), CurrentMipLevel));
 
 						if (const FImgMediaTileSelection* CachedSelection = OutFrame->MipTilesPresent.Find(CurrentMipLevel))
 						{
@@ -240,7 +247,7 @@ bool FExrImgMediaReaderGpu::ReadFrame(int32 FrameId, const TMap<int32, FImgMedia
 						Viewports.Add(MoveTemp(Viewport));
 					}
 
-					ReadResult = ReadTilesCustom(MipDataPtr, BufferSize, ImagePath, FrameId, TileRegionsToRead, ConverterParams, CurrentMipLevel);
+					ReadResult = ReadTiles(MipDataPtr, BufferSize, ImagePath, FrameId, TileRegionsToRead, ConverterParams, CurrentMipLevel, BufferRegionsToCopy);
 				}
 				else
 				{
@@ -253,7 +260,27 @@ bool FExrImgMediaReaderGpu::ReadFrame(int32 FrameId, const TMap<int32, FImgMedia
 
 					ReadResult = ReadInChunks(MipDataPtr, ImagePath, FrameId, CurrentMipDim, BufferSize);
 				}
-
+				if (ReadResult == Success && CVarExrReaderUseUploadHeap.GetValueOnAnyThread())
+				{
+					TArray<int32> InMipLevels;
+					InMipTiles.GetKeys(InMipLevels);
+					ENQUEUE_RENDER_COMMAND(CopyFromUploadBuffer)([SampleConverter, ConverterParams, BufferData, BufferRegionsToCopy](FRHICommandListImmediate& RHICmdList)
+					{
+						SCOPED_DRAW_EVENT(RHICmdList, FExrImgMediaReaderGpu_CopyBuffers);
+						TRACE_CPUPROFILER_EVENT_SCOPE_TEXT(*FString::Printf(TEXT("ExrReaderGpu.StartCopy %d"), ConverterParams->FrameId));
+						if (BufferRegionsToCopy.IsEmpty())
+						{
+							RHICmdList.CopyBuffer(BufferData->UploadBufferRef, BufferData->ShaderAccessBufferRef);
+						}
+						else
+						{
+							for (UE::Math::TIntPoint<int64> Region : BufferRegionsToCopy)
+							{
+								RHICmdList.CopyBufferRegion(BufferData->ShaderAccessBufferRef, Region.X, BufferData->UploadBufferRef, Region.X, Region.Y);
+							}
+						}
+					});
+				}
 				if (ReadResult != Fail)
 				{
 					OutFrame->MipTilesPresent.Emplace(CurrentMipLevel, CurrentTileSelection);
@@ -425,7 +452,7 @@ void FExrImgMediaReaderGpu::CreateSampleConverterCallback
 			FStructuredBufferPoolItemSharedPtr BufferData = MipBuffers[MipLevel];
 			if (BufferData.IsValid())
 			{
-				if (!BufferData->BufferRef->IsValid())
+				if (!BufferData->UploadBufferRef.IsValid() || (CVarExrReaderUseUploadHeap.GetValueOnAnyThread() && !BufferData->ShaderAccessBufferRef->IsValid()))
 				{
 					continue;
 				}
@@ -454,13 +481,14 @@ void FExrImgMediaReaderGpu::CreateSampleConverterCallback
 				if (ConverterParams->FrameInfo.bHasTiles && 
 					(ConverterParams->TileInfoPerMipLevel.Num() > MipLevel && ConverterParams->TileInfoPerMipLevel[MipLevel].Num() > 0))
 				{
-					TRACE_CPUPROFILER_EVENT_SCOPE_TEXT(*FString::Printf(TEXT("FExrImgMediaReaderGpu_TileDesc")));
+					TRACE_CPUPROFILER_EVENT_SCOPE_TEXT(*FString::Printf(TEXT("ExrReaderGpu.TileDesc")));
 
 					FBufferRHIRef BufferRef;
 					FRHIResourceCreateInfo CreateInfo(TEXT("FExrImgMediaReaderGpu_TileDesc"));
 					const uint32 BytesPerElement = sizeof(FExrReader::FTileDesc);
 					const uint32 NumElements = ConverterParams->TileInfoPerMipLevel[MipLevel].Num();
 
+					// This buffer is allocated on already allocated block, therefore the risk of fragmentation is mitigated.
 					BufferRef = RHICreateStructuredBuffer(BytesPerElement, BytesPerElement*NumElements, BUF_ShaderResource | BUF_Dynamic | BUF_FastVRAM, CreateInfo);
 					void* MappedBuffer = RHILockBuffer(BufferRef, 0, NumElements, RLM_WriteOnly);
 					FMemory::Memcpy(MappedBuffer, &ConverterParams->TileInfoPerMipLevel[MipLevel][0], sizeof(FExrReader::FTileDesc) * ConverterParams->TileInfoPerMipLevel[MipLevel].Num());
@@ -469,7 +497,7 @@ void FExrImgMediaReaderGpu::CreateSampleConverterCallback
 					PermutationVector.Set<FExrSwizzlePS::FPartialTiles>(true);
 				}
 
-				Parameters.UnswizzledBuffer = BufferData->ShaderResrouceView;
+				Parameters.UnswizzledBuffer = BufferData->ShaderResourceView;
 
 				FGlobalShaderMap* ShaderMap = GetGlobalShaderMap(GMaxRHIFeatureLevel);
 
@@ -491,7 +519,7 @@ void FExrImgMediaReaderGpu::CreateSampleConverterCallback
 				RHICmdList.EndRenderPass();
 
 				// Mark this render command for this buffer as complete, so we can poll it and transfer later.
-				RHICmdList.WriteGPUFence(BufferData->Fence);
+				RHICmdList.WriteGPUFence(BufferData->RenderFence);
 			}
 
 			
@@ -507,7 +535,7 @@ void FExrImgMediaReaderGpu::CreateSampleConverterCallback
 
 FStructuredBufferPoolItemSharedPtr FExrImgMediaReaderGpu::AllocateGpuBufferFromPool(uint32 AllocSize, bool bWait)
 {
-	TRACE_CPUPROFILER_EVENT_SCOPE_TEXT(*FString::Printf(TEXT("FExrImgMediaReaderGpu_AllocBuffer %d"), AllocSize));
+	TRACE_CPUPROFILER_EVENT_SCOPE_TEXT(*FString::Printf(TEXT("ExrReaderGpu.AllocBuffer %d"), AllocSize));
 
 	// This function is attached to the shared pointer and is used to return any allocated memory to staging pool.
 	auto BufferDeleter = [AllocSize](FStructuredBufferPoolItem* ObjectToDelete) {
@@ -541,10 +569,21 @@ FStructuredBufferPoolItemSharedPtr FExrImgMediaReaderGpu::AllocateGpuBufferFromP
 				FScopeLock ScopeLock(&AllocatorCriticalSecion);
 				SCOPED_DRAW_EVENT(RHICmdList, FExrImgMediaReaderGpu_AllocateBuffer);
 				FRHIResourceCreateInfo CreateInfo(TEXT("FExrImgMediaReaderGpu"));
-				AllocatedBuffer->BufferRef = RHICreateStructuredBuffer(sizeof(uint16) * 2., AllocSize, BUF_ShaderResource | BUF_Dynamic | BUF_FastVRAM, CreateInfo);
-				AllocatedBuffer->MappedBuffer = RHILockBuffer(AllocatedBuffer->BufferRef, 0, AllocSize, RLM_WriteOnly);
-				AllocatedBuffer->ShaderResrouceView = RHICreateShaderResourceView(AllocatedBuffer->BufferRef);
-				AllocatedBuffer->Fence = RHICreateGPUFence(TEXT("BufferNoLongerInUseFence"));
+				AllocatedBuffer->UploadBufferRef = RHICreateStructuredBuffer(sizeof(uint16) * 2., AllocSize, BUF_ShaderResource | BUF_Dynamic | BUF_FastVRAM, CreateInfo);
+				AllocatedBuffer->UploadBufferMapped = RHILockBuffer(AllocatedBuffer->UploadBufferRef, 0, AllocSize, RLM_WriteOnly);
+
+				if (CVarExrReaderUseUploadHeap.GetValueOnAnyThread())
+				{
+					AllocatedBuffer->ShaderAccessBufferRef = RHICreateStructuredBuffer(sizeof(uint16) * 2., AllocSize, BUF_ShaderResource | BUF_FastVRAM, CreateInfo);
+					AllocatedBuffer->ShaderResourceView = RHICreateShaderResourceView(AllocatedBuffer->ShaderAccessBufferRef);
+				}
+				else
+				{
+					AllocatedBuffer->ShaderResourceView = RHICreateShaderResourceView(AllocatedBuffer->UploadBufferRef);
+				}
+
+				AllocatedBuffer->RenderFence = RHICreateGPUFence(TEXT("Exr.BufferNoLongerInUseFence"));
+
 				if (bWait)
 				{
 					bInitDone = true;
@@ -577,7 +616,7 @@ void FExrImgMediaReaderGpu::ReturnGpuBufferToStagingPool(uint32 AllocSize, FStru
 
 			// By this point we don't need a lock because the destructor was already called and it 
 			// is guaranteed that this buffer is no longer used anywhere else.
-			RHIUnlockBuffer(Buffer->BufferRef);
+			RHIUnlockBuffer(Buffer->UploadBufferRef);
 			delete Buffer;
 		});
 	}
@@ -607,7 +646,7 @@ void FExrImgMediaReaderGpu::TransferFromStagingBuffer()
 			for (FStructuredBufferPoolItem* MemoryPoolItem : AllValues)
 			{
 				// Check if fence has signaled. Or otherwise if we are waiting for signal to come through.
-				if (MemoryPoolItem->Fence->Poll() || !MemoryPoolItem->bWillBeSignaled)
+				if (MemoryPoolItem->RenderFence->Poll() || !MemoryPoolItem->bWillBeSignaled)
 				{
 					// If buffer was in use but the fence signaled we need to reset bWillBeSignaled flag.
 					MemoryPoolItem->bWillBeSignaled = false;
diff --git a/Engine/Plugins/Media/ImgMedia/Source/ImgMedia/Private/Readers/ExrImgMediaReaderGpu.h b/Engine/Plugins/Media/ImgMedia/Source/ImgMedia/Private/Readers/ExrImgMediaReaderGpu.h
--- a/Engine/Plugins/Media/ImgMedia/Source/ImgMedia/Private/Readers/ExrImgMediaReaderGpu.h
+++ b/Engine/Plugins/Media/ImgMedia/Source/ImgMedia/Private/Readers/ExrImgMediaReaderGpu.h
@@ -16,23 +16,29 @@ struct FStructuredBufferPoolItem
 {
 	/**
 	* This is the actual buffer reference that we need to keep after it is locked and until it is unlocked.
+	* The buffer is used as an upload heap and will not be accessed by shader if CVarExrReaderUseUploadHeap is set.
 	*/
-	FBufferRHIRef BufferRef;
+	FBufferRHIRef UploadBufferRef;
 
 	/** 
 	* A pointer to mapped GPU memory.
 	*/
-	void* MappedBuffer;
+	void* UploadBufferMapped;
+
+	/**
+	* This buffer is used by the swizzling shader if CVarExrReaderUseUploadHeap is set and UploadBufferRef contents are copied into it.
+	*/
+	FBufferRHIRef ShaderAccessBufferRef;
 
 	/** 
 	* Resource View used by swizzling shader.
 	*/
-	FShaderResourceViewRHIRef ShaderResrouceView;
+	FShaderResourceViewRHIRef ShaderResourceView;
 
 	/** 
 	* A Gpu fence that identifies if this pool item is available for use again.
 	*/
-	FGPUFenceRHIRef Fence;
+	FGPUFenceRHIRef RenderFence;
 
 	/**
 	* This boolean is used as a flag in combination with fences to indicate if rendering thread 
