ImgMedia: Fix visible mip-tiles calculations for partial tiles.

#rnx

diff --git a/Engine/Plugins/Media/ImgMedia/Source/ImgMedia/Private/Assets/ImgMediaMipMapInfo.cpp b/Engine/Plugins/Media/ImgMedia/Source/ImgMedia/Private/Assets/ImgMediaMipMapInfo.cpp
--- a/Engine/Plugins/Media/ImgMedia/Source/ImgMedia/Private/Assets/ImgMediaMipMapInfo.cpp
+++ b/Engine/Plugins/Media/ImgMedia/Source/ImgMedia/Private/Assets/ImgMediaMipMapInfo.cpp
@@ -273,7 +273,7 @@ void FImgMediaMipMapObjectInfo::CalculateVisibleTiles(const TArray<FImgMediaView
 	// We simply add fully visible regions for all mip levels
 	for (int32 MipLevel = 0; MipLevel < InSequenceInfo.NumMipLevels; ++MipLevel)
 	{
-		VisibleTiles.Add(MipLevel, FImgMediaTileSelection::CreateForTargetMipLevel(InSequenceInfo.NumTiles.X, InSequenceInfo.NumTiles.Y, MipLevel, true));
+		VisibleTiles.Add(MipLevel, FImgMediaTileSelection::CreateForTargetMipLevel(InSequenceInfo.TilingDescription.TileNum.X, InSequenceInfo.TilingDescription.TileNum.Y, MipLevel, true));
 	}
 }
 
@@ -340,8 +340,11 @@ namespace {
 				return;
 			}
 
+			const FIntPoint& SequenceTileNum = InSequenceInfo.TilingDescription.TileNum;
+			const FVector2f SequencePartialTileNum = FVector2f(InSequenceInfo.Dim) / FVector2f(InSequenceInfo.TilingDescription.TileSize);
+
 			// To avoid calculating tile corner mip levels multiple times over, we cache them in this array.
-			CornerMipLevelsCached.SetNum((InSequenceInfo.NumTiles.X + 1) * (InSequenceInfo.NumTiles.Y + 1));
+			CornerMipLevelsCached.SetNum((SequenceTileNum.X + 1) * (SequenceTileNum.Y + 1));
 
 			const FTransform MeshTransform = Mesh->GetComponentTransform();
 			const FVector MeshScale = Mesh->GetComponentScale();
@@ -369,8 +372,8 @@ namespace {
 				int MipLevelDiv = 1 << MaxLevel;
 
 				FIntPoint CurrentNumTiles;
-				CurrentNumTiles.X = FMath::Max(1, FMath::CeilToInt(float(InSequenceInfo.NumTiles.X) / MipLevelDiv));
-				CurrentNumTiles.Y = FMath::Max(1, FMath::CeilToInt(float(InSequenceInfo.NumTiles.Y) / MipLevelDiv));
+				CurrentNumTiles.X = FMath::Max(1, FMath::CeilToInt(float(SequenceTileNum.X) / MipLevelDiv));
+				CurrentNumTiles.Y = FMath::Max(1, FMath::CeilToInt(float(SequenceTileNum.Y) / MipLevelDiv));
 
 				// Starting with tiles at the highest mip level
 				TQueue<FIntVector> Tiles;
@@ -391,8 +394,10 @@ namespace {
 					int32 CurrentMipLevel = Tile.Z;
 					MipLevelDiv = 1 << CurrentMipLevel;
 					// Calculate the number of tiles at this mip level
-					CurrentNumTiles.X = FMath::Max(1, FMath::CeilToInt(float(InSequenceInfo.NumTiles.X) / MipLevelDiv));
-					CurrentNumTiles.Y = FMath::Max(1, FMath::CeilToInt(float(InSequenceInfo.NumTiles.Y) / MipLevelDiv));
+					CurrentNumTiles.X = FMath::Max(1, FMath::CeilToInt(float(SequenceTileNum.X) / MipLevelDiv));
+					CurrentNumTiles.Y = FMath::Max(1, FMath::CeilToInt(float(SequenceTileNum.Y) / MipLevelDiv));
+
+					FVector2f CurrentPartialTileNum = SequencePartialTileNum / (float)MipLevelDiv;
 
 					// Exclude subdivided tiles (enqueued below) that are not present (i.e. mipped sequences with odd number of tiles)
 					if (Tile.X >= CurrentNumTiles.X || Tile.Y >= CurrentNumTiles.Y)
@@ -401,12 +406,12 @@ namespace {
 					}
 
 					// Calculate the tile location in world-space
-					float StepX = float(Tile.X + 0.5f) / CurrentNumTiles.X;
-					float StepY = float(Tile.Y + 0.5f) / CurrentNumTiles.Y;
+					float StepX = float(Tile.X + 0.5f) / CurrentPartialTileNum.X;
+					float StepY = float(Tile.Y + 0.5f) / CurrentPartialTileNum.Y;
 					FVector TileCenterWS = PlaneCornerWS + (DirXWS * StepX + DirYWS * StepY);
 
 					// Calculate the tile radius in world space
-					FVector TileSizeWS = (PlaneSize * MeshScale) / FVector(1, CurrentNumTiles.X, CurrentNumTiles.Y);
+					FVector TileSizeWS = (PlaneSize * MeshScale) / FVector(1, CurrentPartialTileNum.X, CurrentPartialTileNum.Y);
 					float TileRadiusInWS = 0.5f * (float)FMath::Sqrt(2 * FMath::Square(TileSizeWS.GetAbsMax()));
 
 					// Now we check if tile spherical bounds are in view.
@@ -424,18 +429,18 @@ namespace {
 								int32 TileCornerY = Tile.Y + CornerY;
 
 								// First we query the cached corner mip levels.
-								int32 MaxCornerX = InSequenceInfo.NumTiles.X + 1;
-								int32 MaxCornerY = InSequenceInfo.NumTiles.Y + 1;
+								int32 MaxCornerX = SequenceTileNum.X + 1;
+								int32 MaxCornerY = SequenceTileNum.Y + 1;
 								FIntPoint BaseLevelCorner;
-								BaseLevelCorner.X = FMath::Clamp(TileCornerX << CurrentMipLevel, 0, InSequenceInfo.NumTiles.X);
-								BaseLevelCorner.Y = FMath::Clamp(TileCornerY << CurrentMipLevel, 0, InSequenceInfo.NumTiles.Y);
+								BaseLevelCorner.X = FMath::Clamp(TileCornerX << CurrentMipLevel, 0, SequenceTileNum.X);
+								BaseLevelCorner.Y = FMath::Clamp(TileCornerY << CurrentMipLevel, 0, SequenceTileNum.Y);
 								bool bValidLevel = GetCachedMipLevel(BaseLevelCorner.X, BaseLevelCorner.Y, MaxCornerX, CalculatedLevel);
 
 								// If not found, calculate and cache it.
 								if (!bValidLevel)
 								{
-									float CornerStepX = TileCornerX / (float)CurrentNumTiles.X;
-									float CornerStepY = TileCornerY / (float)CurrentNumTiles.Y;
+									float CornerStepX = TileCornerX / CurrentPartialTileNum.X;
+									float CornerStepY = TileCornerY / CurrentPartialTileNum.Y;
 									FVector CornersWS = PlaneCornerWS + (DirXWS * CornerStepX + DirYWS * CornerStepY);
 
 									if (CalculateMipLevel(ViewInfo, CornersWS, CornersWS + TexelOffsetXWS, CornersWS + TexelOffsetYWS, CalculatedLevel))
@@ -465,9 +470,9 @@ namespace {
 						// If the lowest (calculated) mip level is below our current mip level, enqueue all 4 sub-tiles for further processing.
 						if (MipLevelRange[0] < CurrentMipLevel)
 						{
-							for (int32 SubY = 0; SubY < FMath::Min(InSequenceInfo.NumTiles.Y, 2); ++SubY)
+							for (int32 SubY = 0; SubY < FMath::Min(SequenceTileNum.Y, 2); ++SubY)
 							{
-								for (int32 SubX = 0; SubX < FMath::Min(InSequenceInfo.NumTiles.X, 2); ++SubX)
+								for (int32 SubX = 0; SubX < FMath::Min(SequenceTileNum.X, 2); ++SubX)
 								{
 									FIntVector SubTile = FIntVector((Tile.X << 1) + SubX, (Tile.Y << 1) + SubY, CurrentMipLevel - 1);
 									Tiles.Enqueue(SubTile);
@@ -561,12 +566,13 @@ namespace {
 			const int32 MaxLevel = InSequenceInfo.NumMipLevels - 1;
 
 			// Include all tiles containted in the visible UV region
-			int32 NumX = InSequenceInfo.NumTiles.X;
-			int32 NumY = InSequenceInfo.NumTiles.Y;
+			const FIntPoint& SequenceTileNum = InSequenceInfo.TilingDescription.TileNum;
+			const FVector2f SequencePartialTileNum = FVector2f(InSequenceInfo.Dim) / FVector2f(InSequenceInfo.TilingDescription.TileSize);
+
 			float PixelDimX = 1.0f / InSequenceInfo.Dim.X;
 			float PixelDimY = 1.0f / InSequenceInfo.Dim.Y;
 
-			FVector ApproxTileSizeWS = MeshTransform.GetScale3D() * (UE_TWO_PI * DefaultSphereRadius) / FMath::Max((float)NumX, (float)NumY);
+			FVector ApproxTileSizeWS = MeshTransform.GetScale3D() * (UE_TWO_PI * DefaultSphereRadius) / FMath::Max(SequencePartialTileNum.X, SequencePartialTileNum.Y);
 			float ApproxTileRadiusInWS = 0.5f * (float)FMath::Sqrt(2 * FMath::Square(ApproxTileSizeWS.GetAbsMax()));
 
 			auto TransformSphericalUVsToLocationWS = [this, &MeshTransform, DefaultSphereRadius](const FVector2D& InUV) -> FVector
@@ -588,11 +594,11 @@ namespace {
 				FConvexVolume ViewFrustum;
 				GetViewFrustumBounds(ViewFrustum, ViewInfo.OverscanViewProjectionMatrix, false, false);
 
-				for (int32 TileY = 0; TileY < NumY; ++TileY)
+				for (int32 TileY = 0; TileY < SequenceTileNum.Y; ++TileY)
 				{
-					for (int32 TileX = 0; TileX < NumX; ++TileX)
+					for (int32 TileX = 0; TileX < SequenceTileNum.X; ++TileX)
 					{
-						FVector2D TileCenterUV = FVector2D(((float)TileX + 0.5f ) / NumX, ((float)TileY + 0.5f) / NumY);
+						FVector2D TileCenterUV = FVector2D(((float)TileX + 0.5f ) / SequencePartialTileNum.X, ((float)TileY + 0.5f) / SequencePartialTileNum.Y);
 						FVector TileCenterWS = TransformSphericalUVsToLocationWS(TileCenterUV);
 
 						// For each tile corner, we include all adjacent tiles
@@ -613,7 +619,7 @@ namespace {
 
 								// As a mitigation for discontinuities at the poles, we artifically increase the max calculated level.
 								// (Note: Using an icosphere would avoid this issue but conflict with the partial sphere feature.)
-								if (TileY == 0 || TileY == NumY - 1)
+								if (TileY == 0 || TileY == SequenceTileNum.Y - 1)
 								{
 									MipLevelRange[1]++;
 								}
@@ -624,7 +630,7 @@ namespace {
 								{
 									if (!VisibleTiles.Contains(Level))
 									{
-										VisibleTiles.Emplace(Level, FImgMediaTileSelection::CreateForTargetMipLevel(InSequenceInfo.NumTiles.X, InSequenceInfo.NumTiles.Y, Level, false));
+										VisibleTiles.Emplace(Level, FImgMediaTileSelection::CreateForTargetMipLevel(SequenceTileNum.X, SequenceTileNum.Y, Level, false));
 									}
 
 									const int MipLevelDiv = 1 << Level;
@@ -761,16 +767,17 @@ void FImgMediaMipMapInfo::ClearAllObjects()
 }
 
 
-void FImgMediaMipMapInfo::SetTextureInfo(FName InSequenceName, int32 InNumMipMaps,
-	const FIntPoint& InNumTiles, const FIntPoint& InSequenceDim)
+void FImgMediaMipMapInfo::SetTextureInfo(FName InSequenceName, int32 InNumMipMaps, const FIntPoint& InSequenceDim,
+	const FMediaTextureTilingDescription& InTilingDescription)
 {
 	SequenceInfo.Name = InSequenceName;
 	SequenceInfo.Dim = InSequenceDim;
+	SequenceInfo.TilingDescription = InTilingDescription;
 
 	// To simplify logic, we assume we always have at least one mip level and one tile.
 	SequenceInfo.NumMipLevels = FMath::Max(1, InNumMipMaps);
-	SequenceInfo.NumTiles.X = FMath::Max(1, InNumTiles.X);
-	SequenceInfo.NumTiles.Y = FMath::Max(1, InNumTiles.Y);
+	SequenceInfo.TilingDescription.TileNum.X = FMath::Max(1, SequenceInfo.TilingDescription.TileNum.X);
+	SequenceInfo.TilingDescription.TileNum.Y = FMath::Max(1, SequenceInfo.TilingDescription.TileNum.Y);
 }
 
 TMap<int32, FImgMediaTileSelection> FImgMediaMipMapInfo::GetVisibleTiles()
diff --git a/Engine/Plugins/Media/ImgMedia/Source/ImgMedia/Private/Assets/ImgMediaMipMapInfo.h b/Engine/Plugins/Media/ImgMedia/Source/ImgMedia/Private/Assets/ImgMediaMipMapInfo.h
--- a/Engine/Plugins/Media/ImgMedia/Source/ImgMedia/Private/Assets/ImgMediaMipMapInfo.h
+++ b/Engine/Plugins/Media/ImgMedia/Source/ImgMedia/Private/Assets/ImgMediaMipMapInfo.h
@@ -5,8 +5,9 @@
 #include "Containers/BitArray.h"
 #include "CoreMinimal.h"
 #include "IMediaOptions.h"
-#include "MediaTextureTracker.h"
+#include "IMediaTextureSample.h"
 #include "ImgMediaSceneViewExtension.h"
+#include "MediaTextureTracker.h"
 #include "Tickable.h"
 #include "UObject/WeakObjectPtr.h"
 #include "UObject/WeakObjectPtrTemplates.h"
@@ -24,10 +25,10 @@ struct FSequenceInfo
 	FName Name;
 	/** Pixel dimensions of this sequence. */
 	FIntPoint Dim;
-	/** Number of tiles in the X, Y directions. */
-	FIntPoint NumTiles;
 	/** Number of mip levels. */
 	int32 NumMipLevels;
+	/** Tiling description. */
+	FMediaTextureTilingDescription TilingDescription;
 	/** 
 	* Check if the sequence has tiles.
 	* 
@@ -35,7 +36,7 @@ struct FSequenceInfo
 	*/
 	FORCEINLINE bool IsTiled() const
 	{
-		return (NumTiles.X > 1) || (NumTiles.Y > 1);
+		return (TilingDescription.TileNum.X > 1) || (TilingDescription.TileNum.Y > 1);
 	}
 };
 
@@ -247,12 +248,12 @@ public:
 	 * Provide information on the texture needed for our image sequence.
 	 *
 	 * @param InSequenceName Name of this sequence.
-	 * @param InNumMipMaps Number of mipmaps in our image sequence.
-	 * @param InNumTiles Number of tiles in our image sequence.
+	 * @param InNumMipMaps Number of mipmaps in our image sequence
 	 * @param InSequenceDim Dimensions of the textures in our image sequence.
+	 * @param InTilingDesc Tiling description of our image sequence.
 	 */
-	void SetTextureInfo(FName InSequenceName, int32 InNumMipMaps, const FIntPoint& InNumTiles,
-		const FIntPoint& InSequenceDim);
+	void SetTextureInfo(FName InSequenceName, int32 InNumMipMaps, const FIntPoint& InSequenceDim,
+		const FMediaTextureTilingDescription& InTilingDesc);
 
 	/**
 	 * Get what mipmap level should be used.
diff --git a/Engine/Plugins/Media/ImgMedia/Source/ImgMedia/Private/Loader/ImgMediaLoader.cpp b/Engine/Plugins/Media/ImgMedia/Source/ImgMedia/Private/Loader/ImgMediaLoader.cpp
--- a/Engine/Plugins/Media/ImgMedia/Source/ImgMedia/Private/Loader/ImgMediaLoader.cpp
+++ b/Engine/Plugins/Media/ImgMedia/Source/ImgMedia/Private/Loader/ImgMediaLoader.cpp
@@ -861,7 +861,7 @@ void FImgMediaLoader::LoadSequence(const FString& SequencePath, const FFrameRate
 		}
 		else
 		{
-			MipMapInfo->SetTextureInfo(SequenceName, GetNumMipLevels(), TilingDescription.TileNum, SequenceDim);
+			MipMapInfo->SetTextureInfo(SequenceName, GetNumMipLevels(), SequenceDim, TilingDescription);
 			if (GetNumMipLevels() > 1)
 			{
 				UncompressedSize = (UncompressedSize * 4) / 3;
