ImgMedia: Minor cleanup to the mipmapinfo class. Tick logic is only applied when tracked objects are present: this prevents needless view info copies and stale debugs logs from polluting the active  level.

#rnx

diff --git a/Engine/Plugins/Media/ImgMedia/Source/ImgMedia/Private/Assets/ImgMediaMipMapInfo.cpp b/Engine/Plugins/Media/ImgMedia/Source/ImgMedia/Private/Assets/ImgMediaMipMapInfo.cpp
--- a/Engine/Plugins/Media/ImgMedia/Source/ImgMedia/Private/Assets/ImgMediaMipMapInfo.cpp
+++ b/Engine/Plugins/Media/ImgMedia/Source/ImgMedia/Private/Assets/ImgMediaMipMapInfo.cpp
@@ -36,13 +36,13 @@ FImgMediaTileSelection::FImgMediaTileSelection(int32 NumTilesX, int32 NumTilesY,
 {
 }
 
-FImgMediaTileSelection FImgMediaTileSelection::CreateForTargetMipLevel(int32 BaseNumTilesX, int32 BaseNumTilesY, int32 TargetMipLevel, bool bDefaultVisibility)
+FImgMediaTileSelection FImgMediaTileSelection::CreateForTargetMipLevel(const FIntPoint& BaseTileNum, int32 TargetMipLevel, bool bDefaultVisibility)
 {
 	ensure(TargetMipLevel >= 0);
 
 	const int MipLevelDiv = 1 << TargetMipLevel;
-	int32 NumTilesX = FMath::Max(1, FMath::CeilToInt(float(BaseNumTilesX) / MipLevelDiv));
-	int32 NumTilesY = FMath::Max(1, FMath::CeilToInt(float(BaseNumTilesY) / MipLevelDiv));
+	int32 NumTilesX = FMath::Max(1, FMath::CeilToInt(float(BaseTileNum.X) / MipLevelDiv));
+	int32 NumTilesY = FMath::Max(1, FMath::CeilToInt(float(BaseTileNum.Y) / MipLevelDiv));
 
 	return FImgMediaTileSelection(NumTilesX, NumTilesY, bDefaultVisibility);
 }
@@ -273,7 +273,7 @@ void FImgMediaMipMapObjectInfo::CalculateVisibleTiles(const TArray<FImgMediaView
 	// We simply add fully visible regions for all mip levels
 	for (int32 MipLevel = 0; MipLevel < InSequenceInfo.NumMipLevels; ++MipLevel)
 	{
-		VisibleTiles.Add(MipLevel, FImgMediaTileSelection::CreateForTargetMipLevel(InSequenceInfo.TilingDescription.TileNum.X, InSequenceInfo.TilingDescription.TileNum.Y, MipLevel, true));
+		VisibleTiles.Add(MipLevel, FImgMediaTileSelection::CreateForTargetMipLevel(InSequenceInfo.TilingDescription.TileNum, MipLevel, true));
 	}
 }
 
@@ -630,7 +630,7 @@ namespace {
 								{
 									if (!VisibleTiles.Contains(Level))
 									{
-										VisibleTiles.Emplace(Level, FImgMediaTileSelection::CreateForTargetMipLevel(SequenceTileNum.X, SequenceTileNum.Y, Level, false));
+										VisibleTiles.Emplace(Level, FImgMediaTileSelection::CreateForTargetMipLevel(SequenceTileNum, Level, false));
 									}
 
 									const int MipLevelDiv = 1 << Level;
@@ -820,19 +820,19 @@ void FImgMediaMipMapInfo::Tick(float DeltaTime)
 {
 	FScopeLock Lock(&InfoCriticalSection);
 
-	const TSharedPtr<FImgMediaSceneViewExtension, ESPMode::ThreadSafe>& SVE = IImgMediaModule::Get().GetSceneViewExtension();
-	if (SVE.IsValid())
-	{
-		ViewInfos = SVE->GetViewInfos();
-	}
-
 	// Let the cache update this frame.
 	bIsCacheValid = false;
 
-	// Display debug?
-	if (CVarImgMediaMipMapDebugEnable.GetValueOnGameThread())
+	if (HasObjects())
 	{
-		if (GEngine != nullptr)
+		const TSharedPtr<FImgMediaSceneViewExtension, ESPMode::ThreadSafe>& SVE = IImgMediaModule::Get().GetSceneViewExtension();
+		if (SVE.IsValid())
+		{
+			ViewInfos = SVE->GetViewInfos();
+		}
+
+		// Display debug?
+		if (CVarImgMediaMipMapDebugEnable.GetValueOnGameThread() && GEngine != nullptr)
 		{
 			TSet<int32> VisibleMips;
 			int32 NumVisibleTiles = 0;
diff --git a/Engine/Plugins/Media/ImgMedia/Source/ImgMedia/Private/Assets/ImgMediaMipMapInfo.h b/Engine/Plugins/Media/ImgMedia/Source/ImgMedia/Private/Assets/ImgMediaMipMapInfo.h
--- a/Engine/Plugins/Media/ImgMedia/Source/ImgMedia/Private/Assets/ImgMediaMipMapInfo.h
+++ b/Engine/Plugins/Media/ImgMedia/Source/ImgMedia/Private/Assets/ImgMediaMipMapInfo.h
@@ -62,7 +62,7 @@ struct FImgMediaTileSelection
 	 * @param TargetMipLevel Higher target mip level for the selection (usually 1 and above).
 	 * @param bDefaultVisibility Optional visibility for the entire region.
 	 */
-	static FImgMediaTileSelection CreateForTargetMipLevel(int32 BaseNumTilesX, int32 BaseNumTilesY, int32 TargetMipLevel, bool bDefaultVisibility = false);
+	static FImgMediaTileSelection CreateForTargetMipLevel(const FIntPoint& BaseTileNum, int32 TargetMipLevel, bool bDefaultVisibility = false);
 
 	FImgMediaTileSelection() = default;
 	~FImgMediaTileSelection() noexcept = default;
diff --git a/Engine/Plugins/Media/ImgMedia/Source/ImgMedia/Private/Loader/ImgMediaLoader.cpp b/Engine/Plugins/Media/ImgMedia/Source/ImgMedia/Private/Loader/ImgMediaLoader.cpp
--- a/Engine/Plugins/Media/ImgMedia/Source/ImgMedia/Private/Loader/ImgMediaLoader.cpp
+++ b/Engine/Plugins/Media/ImgMedia/Source/ImgMedia/Private/Loader/ImgMediaLoader.cpp
@@ -1311,7 +1311,7 @@ void FImgMediaLoader::GetDesiredMipTiles(int32 FrameIndex, TMap<int32, FImgMedia
 		//Fallback case where we activate all mips and tiles.
 		for (int32 MipLevel = 0; MipLevel < GetNumMipLevels(); ++MipLevel)
 		{
-			OutMipsAndTiles.Emplace(MipLevel, FImgMediaTileSelection::CreateForTargetMipLevel(GetNumTilesX(), GetNumTilesY(), MipLevel, true));
+			OutMipsAndTiles.Emplace(MipLevel, FImgMediaTileSelection::CreateForTargetMipLevel(TilingDescription.TileNum, MipLevel, true));
 		}
 	}
 }
