ProcessImages: Fixed texture size being wrong occasionally when using the player.

#rnx

diff --git a/Engine/Plugins/Media/ImgMedia/Source/ImgMediaEditor/Private/Widgets/SImgMediaProcessImages.cpp b/Engine/Plugins/Media/ImgMedia/Source/ImgMediaEditor/Private/Widgets/SImgMediaProcessImages.cpp
--- a/Engine/Plugins/Media/ImgMedia/Source/ImgMediaEditor/Private/Widgets/SImgMediaProcessImages.cpp
+++ b/Engine/Plugins/Media/ImgMedia/Source/ImgMediaEditor/Private/Widgets/SImgMediaProcessImages.cpp
@@ -899,59 +899,63 @@ void SImgMediaProcessImages::HandleProcessing()
 					FrameDuration = FTimespan::FromSeconds(1.0f / FrameRate);
 				}
 
-				// Copy media texture to our render target.
-				DrawTextureToRenderTarget();
-
-				// Process this render.
-				TArray64<uint8> RawData;
-				bool bReadSuccess = FImageUtils::GetRawData(RenderTarget, RawData);
-				if (bReadSuccess)
+				// Did we get a render target?
+				if (RenderTarget != nullptr)
 				{
-					int32 Width = RenderTarget->GetSurfaceWidth();
-					int32 Height = RenderTarget->GetSurfaceHeight();
-					int32 BitDepth = 16;
-					bool bUseCustomFormat = Options->bUseCustomFormat;
-					int32 InTileWidth = Options->bEnableTiling ? Options->TileSizeX : 0;
-					int32 InTileHeight = Options->bEnableTiling ? Options->TileSizeY : 0;
-					int32 TileBorder = 0; // Note: virtual texture support is shelved for now.
-					bool bEnableMips = Options->bEnableMipMapping;
-					bool bHasAlphaChannel = false;
-					FString OutPath = Options->OutputPath.Path;
-					FString FileName = FString::Printf(TEXT("image%05d.exr"), CurrentFrameIndex);
-					FString Name = FPaths::Combine(OutPath, FileName);
+					// Copy media texture to our render target.
+					DrawTextureToRenderTarget();
 
-					Async(EAsyncExecution::Thread, [this, RawData = MoveTemp(RawData), Width, Height, BitDepth,
-						InTileWidth, InTileHeight, TileBorder, bEnableMips,
-						bHasAlphaChannel, Name]() mutable
+					// Process this render.
+					TArray64<uint8> RawData;
+					bool bReadSuccess = FImageUtils::GetRawData(RenderTarget, RawData);
+					if (bReadSuccess)
 					{
-						ProcessImageCustomRawData(RawData, Width, Height, BitDepth,
+						int32 Width = RenderTarget->GetSurfaceWidth();
+						int32 Height = RenderTarget->GetSurfaceHeight();
+						int32 BitDepth = 16;
+						bool bUseCustomFormat = Options->bUseCustomFormat;
+						int32 InTileWidth = Options->bEnableTiling ? Options->TileSizeX : 0;
+						int32 InTileHeight = Options->bEnableTiling ? Options->TileSizeY : 0;
+						int32 TileBorder = 0; // Note: virtual texture support is shelved for now.
+						bool bEnableMips = Options->bEnableMipMapping;
+						bool bHasAlphaChannel = false;
+						FString OutPath = Options->OutputPath.Path;
+						FString FileName = FString::Printf(TEXT("image%05d.exr"), CurrentFrameIndex);
+						FString Name = FPaths::Combine(OutPath, FileName);
+
+						Async(EAsyncExecution::Thread, [this, RawData = MoveTemp(RawData), Width, Height, BitDepth,
 							InTileWidth, InTileHeight, TileBorder, bEnableMips,
-							bHasAlphaChannel, Name);
-					});
-				}
-				else
-				{
-					UE_LOG(LogImgMediaEditor, Error, TEXT("ProcessImages failed to get raw data."));
-				}
+							bHasAlphaChannel, Name]() mutable
+						{
+							ProcessImageCustomRawData(RawData, Width, Height, BitDepth,
+								InTileWidth, InTileHeight, TileBorder, bEnableMips,
+								bHasAlphaChannel, Name);
+						});
+					}
+					else
+					{
+						UE_LOG(LogImgMediaEditor, Error, TEXT("ProcessImages failed to get raw data."));
+					}
 
-				// Update notification.
-				if (ConfirmNotification.IsValid())
-				{
-					ConfirmNotification->SetText(
-						FText::Format(LOCTEXT("ImgMediaCompleted2", "ImgMedia Completed {0}"),
-							FText::AsNumber(CurrentFrameIndex)));
-				}
+					// Update notification.
+					if (ConfirmNotification.IsValid())
+					{
+						ConfirmNotification->SetText(
+							FText::Format(LOCTEXT("ImgMediaCompleted2", "ImgMedia Completed {0}"),
+								FText::AsNumber(CurrentFrameIndex)));
+					}
 
-				// Next frame.
-				CurrentTime += FrameDuration;
-				CurrentFrameIndex++;
-				if (CurrentTime >= MediaPlayer->GetDuration())
-				{
-					bShouldExit = true;
-				}
-				else
-				{
-					MediaPlayer->SetBlockOnTimeRange(TRange<FTimespan>(CurrentTime, CurrentTime + FrameDuration));
+					// Next frame.
+					CurrentTime += FrameDuration;
+					CurrentFrameIndex++;
+					if (CurrentTime >= MediaPlayer->GetDuration())
+					{
+						bShouldExit = true;
+					}
+					else
+					{
+						MediaPlayer->SetBlockOnTimeRange(TRange<FTimespan>(CurrentTime, CurrentTime + FrameDuration));
+					}
 				}
 			}
 		}
@@ -983,10 +987,11 @@ void SImgMediaProcessImages::HandleProcessing()
 
 void SImgMediaProcessImages::CreateRenderTarget()
 {
-	if (MediaTexture != nullptr)
+	FIntPoint VideoDim = MediaPlayer->GetVideoTrackDimensions(INDEX_NONE, INDEX_NONE);
+	if (VideoDim.X > 0)
 	{
-		int32 Width = MediaTexture->GetWidth();
-		int32 Height = MediaTexture->GetHeight();
+		int32 Width = VideoDim.X;
+		int32 Height = VideoDim.Y;
 
 		RenderTarget = NewObject<UTextureRenderTarget2D>(GetTransientPackage(), TEXT("ImgMediaProcessImages"));
 		RenderTarget->RenderTargetFormat = RTF_RGBA16f;
