ImgMedia: Print to screen additional view description, origin and direction for improved debugging.

#rnx

diff --git a/Engine/Plugins/Media/ImgMedia/Source/ImgMedia/Private/ImgMediaSceneViewExtension.cpp b/Engine/Plugins/Media/ImgMedia/Source/ImgMedia/Private/ImgMediaSceneViewExtension.cpp
--- a/Engine/Plugins/Media/ImgMedia/Source/ImgMedia/Private/ImgMediaSceneViewExtension.cpp
+++ b/Engine/Plugins/Media/ImgMedia/Source/ImgMedia/Private/ImgMediaSceneViewExtension.cpp
@@ -94,6 +94,16 @@ void FImgMediaSceneViewExtension::BeginRenderViewFamily(FSceneViewFamily& InView
 			Info.MaterialTextureMipBias = 0.0f;
 		}
 
+#if !(UE_BUILD_SHIPPING || UE_BUILD_TEST)
+		static const auto CVarMipMapDebug = IConsoleManager::Get().FindConsoleVariable(TEXT("ImgMedia.MipMapDebug"));
+
+		if (GEngine != nullptr && CVarMipMapDebug != nullptr && CVarMipMapDebug->GetBool())
+		{
+			const FString ViewName = InViewFamily.ProfileDescription.IsEmpty() ? TEXT("View"): InViewFamily.ProfileDescription;
+			GEngine->AddOnScreenDebugMessage(-1, 0.0f, FColor::Cyan, *FString::Printf(TEXT("%s location: [%s], direction: [%s]"), *ViewName, *Info.Location.ToString(), *View->GetViewDirection().ToString()));
+		}
+#endif
+
 		CachedViewInfos.Add(MoveTemp(Info));
 	}
 }
