ImgMedia: Fix for relative media sources older than 4.27 (before the IsPathRelativeToProjectRoot property was added), where paths were not relative to the project root folder (for example ./Movie/Sequence/..). This change maintains the previously false default value and adjusts older relative paths on load if they were not project-relative.

#rnx

diff --git a/Engine/Plugins/Media/ImgMedia/Source/ImgMedia/Private/Assets/ImgMediaSource.cpp b/Engine/Plugins/Media/ImgMedia/Source/ImgMedia/Private/Assets/ImgMediaSource.cpp
--- a/Engine/Plugins/Media/ImgMedia/Source/ImgMedia/Private/Assets/ImgMediaSource.cpp
+++ b/Engine/Plugins/Media/ImgMedia/Source/ImgMedia/Private/Assets/ImgMediaSource.cpp
@@ -9,13 +9,14 @@
 
 #include "HAL/FileManager.h"
 #include "Misc/Paths.h"
+#include "UObject/UE5MainStreamObjectVersion.h"
 
 
 /* UImgMediaSource structors
  *****************************************************************************/
 
 UImgMediaSource::UImgMediaSource()
-	: IsPathRelativeToProjectRoot_DEPRECATED(true)
+	: IsPathRelativeToProjectRoot_DEPRECATED(false)
 	, FrameRateOverride(0, 0)
 	, bFillGapsInSequence(true)
 	, MipMapInfo(MakeShared<FImgMediaMipMapInfo, ESPMode::ThreadSafe>())
@@ -238,17 +239,26 @@ void UImgMediaSource::Serialize(FArchive& Ar)
 {
 	Super::Serialize(Ar);
 
-	if (Ar.IsLoading() && !IsPathRelativeToProjectRoot_DEPRECATED)
+#if WITH_EDITOR
+	Ar.UsingCustomVersion(FUE5MainStreamObjectVersion::GUID);
+
+	if (Ar.CustomVer(FUE5MainStreamObjectVersion::GUID) < FUE5MainStreamObjectVersion::ImgMediaPathResolutionWithEngineOrProjectTokens)
 	{
-		// This is an object that was saved with the old value, so we need to convert the path accordingly
+		if (Ar.IsLoading() && !IsPathRelativeToProjectRoot_DEPRECATED)
+		{
+			// This is an object that was saved with the old value (or before the property was added), so we need to convert the path accordingly
 
-		IsPathRelativeToProjectRoot_DEPRECATED = true;
+			IsPathRelativeToProjectRoot_DEPRECATED = true;
 
-		if (FPaths::IsRelative(SequencePath.Path))
-		{
-			SequencePath.Path = FString::Printf(TEXT("Content/%s"), *SequencePath.Path);
+			if (FPaths::IsRelative(SequencePath.Path))
+			{
+				SequencePath.Path = FString::Printf(TEXT("Content/%s"), *SequencePath.Path);
+
+				SequencePath.Path = UImgMediaSource::SanitizeTokenizedSequencePath(SequencePath.Path);
+			}
 		}
 	}
+#endif
 }
 
 bool UImgMediaSource::IsPathUnderBasePath(const FString& InPath, const FString& InBasePath, FString& OutRelativePath)
diff --git a/Engine/Source/Runtime/Core/Public/UObject/UE5MainStreamObjectVersion.h b/Engine/Source/Runtime/Core/Public/UObject/UE5MainStreamObjectVersion.h
--- a/Engine/Source/Runtime/Core/Public/UObject/UE5MainStreamObjectVersion.h
+++ b/Engine/Source/Runtime/Core/Public/UObject/UE5MainStreamObjectVersion.h
@@ -259,6 +259,9 @@ struct CORE_API FUE5MainStreamObjectVersion
 		// Moving all distortion handling to Lens Component
 		LensComponentDistortion,
 
+		// Updated image media source path resolution logic
+		ImgMediaPathResolutionWithEngineOrProjectTokens,
+
 		// -----<new versions can be added above this line>-------------------------------------------------
 		VersionPlusOne,
 		LatestVersion = VersionPlusOne - 1
