ImgMedia: concurrency improvements for mip/tile calculations. Removed unsafe objects getter.

#rnx

diff --git a/Engine/Plugins/Media/ImgMedia/Source/ImgMedia/Private/Assets/ImgMediaMipMapInfo.cpp b/Engine/Plugins/Media/ImgMedia/Source/ImgMedia/Private/Assets/ImgMediaMipMapInfo.cpp
--- a/Engine/Plugins/Media/ImgMedia/Source/ImgMedia/Private/Assets/ImgMediaMipMapInfo.cpp
+++ b/Engine/Plugins/Media/ImgMedia/Source/ImgMedia/Private/Assets/ImgMediaMipMapInfo.cpp
@@ -676,6 +676,8 @@ void FImgMediaMipMapInfo::AddObject(AActor* InActor, float LODBias, EMediaTextur
 {
 	if (InActor != nullptr)
 	{
+		FScopeLock Lock(&ObjectsCriticalSection);
+
 		UMeshComponent* MeshComponent = Cast<UMeshComponent>(InActor->FindComponentByClass(UMeshComponent::StaticClass()));
 		if (MeshComponent != nullptr)
 		{
@@ -699,7 +701,7 @@ void FImgMediaMipMapInfo::RemoveObject(AActor* InActor)
 {
 	if (InActor != nullptr)
 	{
-		FScopeLock Lock(&InfoCriticalSection);
+		FScopeLock Lock(&ObjectsCriticalSection);
 
 		for (int Index = 0; Index < Objects.Num(); ++Index)
 		{
@@ -760,12 +762,13 @@ void FImgMediaMipMapInfo::RemoveObjectsUsingThisMediaTexture(UMediaTexture* InMe
 
 void FImgMediaMipMapInfo::ClearAllObjects()
 {
-	FScopeLock Lock(&InfoCriticalSection);
+	FScopeLock Lock(&ObjectsCriticalSection);
 
 	for (FImgMediaMipMapObjectInfo* Info : Objects)
 	{
 		delete Info;
 	}
+
 	Objects.Empty();
 }
 
@@ -789,34 +792,43 @@ TMap<int32, FImgMediaTileSelection> FImgMediaMipMapInfo::GetVisibleTiles()
 	// So no need for thread safety here with regards to this function.
 	// However the Tick is called from a different thread so care must still be taken when
 	// accessing things that are modified by code external to this function.
+
+	FScopeLock Lock(&InfoCriticalSection);
 	
 	// Do we need to update the cache?
 	if (bIsCacheValid == false)
 	{
-		UpdateMipLevelCache();
-	}
-
-	return CachedVisibleTiles;
-}
-
-void FImgMediaMipMapInfo::UpdateMipLevelCache()
-{
-	SCOPE_CYCLE_COUNTER(STAT_ImgMedia_MipMapUpdateCache);
-
-	{
-		FScopeLock Lock(&InfoCriticalSection);
+		SCOPE_CYCLE_COUNTER(STAT_ImgMedia_MipMapUpdateCache);
 
 		CachedVisibleTiles.Reset();
 
+		FScopeLock LockObjects(&ObjectsCriticalSection);
+
 		// Loop over all objects.
 		for (FImgMediaMipMapObjectInfo* ObjectInfo : Objects)
 		{
 			ObjectInfo->CalculateVisibleTiles(ViewInfos, SequenceInfo, CachedVisibleTiles);
 		}
+
+		// Mark cache as valid.
+		bIsCacheValid = true;
 	}
-	
-	// Mark cache as valid.
-	bIsCacheValid = true;
+
+	return CachedVisibleTiles;
+}
+
+
+/**
+* Check if any scene objects are using our img sequence.
+*
+* @return True if any object is active.
+*/
+
+bool FImgMediaMipMapInfo::HasObjects() const
+{
+	FScopeLock Lock(&ObjectsCriticalSection);
+
+	return Objects.Num() > 0;
 }
 
 void FImgMediaMipMapInfo::Tick(float DeltaTime)
diff --git a/Engine/Plugins/Media/ImgMedia/Source/ImgMedia/Private/Assets/ImgMediaMipMapInfo.h b/Engine/Plugins/Media/ImgMedia/Source/ImgMedia/Private/Assets/ImgMediaMipMapInfo.h
--- a/Engine/Plugins/Media/ImgMedia/Source/ImgMedia/Private/Assets/ImgMediaMipMapInfo.h
+++ b/Engine/Plugins/Media/ImgMedia/Source/ImgMedia/Private/Assets/ImgMediaMipMapInfo.h
@@ -267,19 +267,12 @@ public:
 	 */
 	TMap<int32, FImgMediaTileSelection> GetVisibleTiles();
 	
-	/**
-	 * Get information on objects that are using our textures.
-	 * 
-	 * @return Array of objects.
-	 */
-	const TArray<FImgMediaMipMapObjectInfo*>& GetObjects() const { return Objects; }
-
 	/**
 	 * Check if any scene objects are using our img sequence.
 	 *
 	 * @return True if any object is active.
 	 */
-	bool HasObjects() const { return Objects.Num() > 0; }
+	bool HasObjects() const;
 
 	//~ FTickableGameObject interface
 	virtual void Tick(float DeltaTime) override;
@@ -287,10 +280,6 @@ public:
 	virtual bool IsTickableInEditor() const override { return true; }
 
 protected:
-	/**
-	 * Performs mipmap calculations and caches the data.
-	 */
-	void UpdateMipLevelCache();
 
 	/** Array of objects that are using our img sequence. */
 	TArray<FImgMediaMipMapObjectInfo*> Objects;
@@ -304,6 +293,9 @@ protected:
 	/** Protects info variables. */
 	FCriticalSection InfoCriticalSection;
 
+	/** Protects tracked objects. */
+	mutable FCriticalSection ObjectsCriticalSection;
+
 	/** Desired mipmap level at this current time. */
 	TMap<int32, FImgMediaTileSelection> CachedVisibleTiles;
 
