ProcessImages: Fixed mip levels being incorrect when the previous mip level's width does not divide evenly by 2.

#rnx

diff --git a/Engine/Plugins/Media/ImgMedia/Source/ImgMediaEditor/Private/Widgets/SImgMediaProcessImages.cpp b/Engine/Plugins/Media/ImgMedia/Source/ImgMediaEditor/Private/Widgets/SImgMediaProcessImages.cpp
--- a/Engine/Plugins/Media/ImgMedia/Source/ImgMediaEditor/Private/Widgets/SImgMediaProcessImages.cpp
+++ b/Engine/Plugins/Media/ImgMedia/Source/ImgMediaEditor/Private/Widgets/SImgMediaProcessImages.cpp
@@ -541,6 +541,8 @@ void SImgMediaProcessImages::ProcessImageCustomRawData(TArray64<uint8>& RawData,
 	int32 NumMips = OutFile.GetNumberOfMipLevels();
 	int32 MipSourceWidth = Width;
 	int32 MipSourceHeight = Height;
+	int32 LastMipWidth = 0;
+	int32 LastMipHeight = 0;
 	for (int32 MipLevel = 0; MipLevel < NumMips; MipLevel++)
 	{
 		int32 MipWidth = OutFile.GetMipWidth(MipLevel);
@@ -561,7 +563,7 @@ void SImgMediaProcessImages::ProcessImageCustomRawData(TArray64<uint8>& RawData,
 		{
 			TRACE_CPUPROFILER_EVENT_SCOPE(SImgMediaProcessImages::ProcessImageCustom:GenerateMipData);
 			int32 SourceStrideX = NumChannels;
-			int32 SourceStrideY = MipWidth * NumChannels * 2;
+			int32 SourceStrideY = LastMipWidth * NumChannels;
 			for (int32 PixelY = 0; PixelY < MipHeight; ++PixelY)
 			{
 				for (int32 PixelX = 0; PixelX < MipWidth; ++PixelX)
@@ -570,7 +572,7 @@ void SImgMediaProcessImages::ProcessImageCustomRawData(TArray64<uint8>& RawData,
 					for (int32 Channel = 0; Channel < NumChannels; ++Channel)
 					{
 						// Box filter.
-						int32 SourceWidth = MipWidth * 2;
+						int32 SourceWidth = LastMipWidth;
 						int32 SourcePixelOffset = (PixelX + PixelY * SourceWidth) * NumChannels * 2 +
 							Channel;
 						FFloat16 SourcePixels[4];
@@ -674,6 +676,8 @@ void SImgMediaProcessImages::ProcessImageCustomRawData(TArray64<uint8>& RawData,
 		CurrentMipBufferIndex ^= 1;
 		MipSourceHeight /= 2;
 		MipSourceWidth /= 2;
+		LastMipWidth = MipWidth;
+		LastMipHeight = MipHeight;
 	}
 
 #else // IMGMEDIAEDITOR_EXR_SUPPORTED_PLATFORM
