ImgMedia: Fix for image sequences without 0 prefixes. Sorting is done on the frame number when available, not the string.

#rnx

diff --git a/Engine/Plugins/Media/ImgMedia/Source/ImgMedia/Private/Loader/ImgMediaLoader.cpp b/Engine/Plugins/Media/ImgMedia/Source/ImgMedia/Private/Loader/ImgMediaLoader.cpp
--- a/Engine/Plugins/Media/ImgMedia/Source/ImgMedia/Private/Loader/ImgMediaLoader.cpp
+++ b/Engine/Plugins/Media/ImgMedia/Source/ImgMedia/Private/Loader/ImgMediaLoader.cpp
@@ -5,6 +5,7 @@
 #include "ImgMediaPrivate.h"
 
 #include "Algo/Reverse.h"
+#include "Containers/SortedMap.h"
 #include "Misc/FrameRate.h"
 #include "HAL/FileManager.h"
 #include "HAL/PlatformProcess.h"
@@ -942,34 +943,54 @@ void FImgMediaLoader::FindFiles(const FString& SequencePath, TArray<FString>& Ou
 
 	UE_LOG(LogImgMedia, Verbose, TEXT("Loader %p: Found %i image files in %s"), this, FoundFiles.Num(), *SequencePath);
 
-	FoundFiles.Sort();
+	TArray<FString> UnnumberedFiles;
+	TSortedMap<int32, FString> SortedFoundFiles;
+	SortedFoundFiles.Reserve(FoundFiles.Num());
+
+	int32 FrameNumber;
+	for (FString& File : FoundFiles)
+	{
+		if (GetNumberAtEndOfString(FrameNumber, File))
+		{
+			SortedFoundFiles.Add(FrameNumber, File);
+		}
+		else
+		{
+			UnnumberedFiles.Add(File);
+		}
+	}
 
 	int32 LastIndex = -1;
 	FString EmptyString;
-	for (const auto& File : FoundFiles)
+	for (const auto& Pair : SortedFoundFiles)
 	{
 		// Can we fill in gaps in the sequence?
 		if (bFillGapsInSequence)
 		{
 			// Get the index of this file.
-			int32 ThisIndex = -1;
-			if (GetNumberAtEndOfString(ThisIndex, File))
+			int32 ThisIndex = Pair.Key;
+
+			// Fill in any gaps from the last frame.
+			if ((LastIndex != -1) && (LastIndex < ThisIndex - 1))
 			{
-				// Fill in any gaps from the last frame.
-				if ((LastIndex != -1) && (LastIndex < ThisIndex - 1))
+				for (int32 Index = LastIndex + 1; Index < ThisIndex; ++Index)
 				{
-					for (int32 Index = LastIndex + 1; Index < ThisIndex; ++Index)
-					{
-						OutputPaths.Add(EmptyString);
-					}
+					OutputPaths.Add(EmptyString);
 				}
 			}
 
 			LastIndex = ThisIndex;
 		}
 
-		FString FullPath = FPaths::Combine(SequencePath, File);
-		OutputPaths.Add(FullPath);
+		OutputPaths.Add(FPaths::Combine(SequencePath, Pair.Value));
+	}
+
+	// We still want to support unnumbered files, so we append them to the numbered list.
+	UnnumberedFiles.Sort();
+
+	for (const FString& UnnumberedFile : UnnumberedFiles)
+	{
+		OutputPaths.Add(FPaths::Combine(SequencePath, UnnumberedFile));
 	}
 
 }
