ImgMedia: Removing stale code for old single-file-per-tile path from loader.

#rnx

diff --git a/Engine/Plugins/Media/ImgMedia/Source/ImgMedia/Private/Loader/ImgMediaLoader.cpp b/Engine/Plugins/Media/ImgMedia/Source/ImgMedia/Private/Loader/ImgMediaLoader.cpp
--- a/Engine/Plugins/Media/ImgMedia/Source/ImgMedia/Private/Loader/ImgMediaLoader.cpp
+++ b/Engine/Plugins/Media/ImgMedia/Source/ImgMedia/Private/Loader/ImgMediaLoader.cpp
@@ -91,7 +91,6 @@ FImgMediaLoader::FImgMediaLoader(const TSharedRef<FImgMediaScheduler, ESPMode::T
 	, ImageWrapperModule(FModuleManager::LoadModuleChecked<IImageWrapperModule>("ImageWrapper"))
 	, Initialized(false)
 	, bFillGapsInSequence(bInFillGapsInSequence)
-	, bIsTiled(false)
 	, TilingDescription()
 	, NumLoadAhead(0)
 	, NumLoadBehind(0)
@@ -925,22 +924,6 @@ void FImgMediaLoader::FindFiles(const FString& SequencePath, TArray<FString>& Ou
 	TArray<FString> FoundFiles;
 	IFileManager::Get().FindFiles(FoundFiles, *SequencePath, TEXT("*"));
 
-	// If we did not find any files, then maybe we have tile directories.
-	// If we already have tiles then don't do anything.
-	if (bIsTiled == false)
-	{
-		bIsTiled = FoundFiles.Num() == 0;
-		if (bIsTiled)
-		{
-			IFileManager::Get().FindFiles(FoundFiles, *(SequencePath / TEXT("*")), false, true);
-			if (FoundFiles.Num() == 0)
-			{
-				// We really got nothing.
-				bIsTiled = false;
-			}
-		}
-	}
-
 	UE_LOG(LogImgMedia, Verbose, TEXT("Loader %p: Found %i image files in %s"), this, FoundFiles.Num(), *SequencePath);
 
 	TArray<FString> UnnumberedFiles;
diff --git a/Engine/Plugins/Media/ImgMedia/Source/ImgMedia/Private/Loader/ImgMediaLoader.h b/Engine/Plugins/Media/ImgMedia/Source/ImgMedia/Private/Loader/ImgMediaLoader.h
--- a/Engine/Plugins/Media/ImgMedia/Source/ImgMedia/Private/Loader/ImgMediaLoader.h
+++ b/Engine/Plugins/Media/ImgMedia/Source/ImgMedia/Private/Loader/ImgMediaLoader.h
@@ -486,9 +486,6 @@ private:
 	/** If true, then any gaps in the sequence will be filled with blank frames. */
 	const bool bFillGapsInSequence;
 
-	/** True if this sequence contains tiles. */
-	bool bIsTiled;
-
 	/** Tiling description. */
 	FMediaTextureTilingDescription TilingDescription;
 
