ProcessImages: Non custom output has the same features as custom now (e.g. tiles, mips, etc).

#rnx

diff --git a/Engine/Plugins/Media/ImgMedia/Source/ImgMediaEditor/Private/Widgets/SImgMediaProcessImages.cpp b/Engine/Plugins/Media/ImgMedia/Source/ImgMediaEditor/Private/Widgets/SImgMediaProcessImages.cpp
--- a/Engine/Plugins/Media/ImgMedia/Source/ImgMediaEditor/Private/Widgets/SImgMediaProcessImages.cpp
+++ b/Engine/Plugins/Media/ImgMedia/Source/ImgMediaEditor/Private/Widgets/SImgMediaProcessImages.cpp
@@ -295,17 +295,9 @@ void SImgMediaProcessImages::ProcessAllImages()
 
 					// Import this image.
 					FString Name = FPaths::Combine(OutPath, FileName);
-					if (bUseCustomFormat)
-					{
-						ProcessImageCustom(ImageWrapper, InTileWidth, InTileHeight, TileBorder,
-							bEnableMips, bHasAlphaChannel, Name);
+					ProcessImageCustom(ImageWrapper, InTileWidth, InTileHeight, TileBorder,
+						bEnableMips, bHasAlphaChannel, Name, bUseCustomFormat);
 
-					}
-					else
-					{
-						Name = FPaths::ChangeExtension(Name, TEXT(""));
-						ProcessImage(ImageWrapper, InTileWidth, InTileHeight, Name, Ext);
-					}
 					NumActive--;
 				});
 
@@ -354,73 +346,10 @@ bool SImgMediaProcessImages::HasAlphaChannel(const FString& Ext, const FString&
 	return bHasAlpha;
 }
 
-void SImgMediaProcessImages::ProcessImage(TSharedPtr<IImageWrapper>& InImageWrapper,
-	int32 InTileWidth, int32 InTileHeight, const FString& InName, const FString& FileExtension)
-{
-	// Get image data.
-	ERGBFormat Format = InImageWrapper->GetFormat();
-	int32 Width = InImageWrapper->GetWidth();
-	int32 Height = InImageWrapper->GetHeight();
-	int32 BitDepth = InImageWrapper->GetBitDepth();
-	TArray64<uint8> RawData;
-	InImageWrapper->GetRaw(Format, BitDepth, RawData);
-
-	int32 NumTilesX = InTileWidth > 0 ? Width / InTileWidth : 1;
-	int32 NumTilesY = InTileHeight > 0 ? Height / InTileHeight : 1;
-	int32 TileWidth = Width / NumTilesX;
-	int32 TileHeight = Height / NumTilesY;
-	int32 BytesPerPixel = RawData.Num() / (Width * Height);
-	TArray64<uint8> TileRawData;
-	TileRawData.AddZeroed(TileWidth * TileHeight * BytesPerPixel);
-	bool bIsTiled = (NumTilesX > 1) || (NumTilesY > 1);
-
-	// Create a directory if we have tiles.
-	FString FileName;
-	if (bIsTiled)
-	{
-		IPlatformFile& PlatformFile = FPlatformFileManager::Get().GetPlatformFile();
-		PlatformFile.CreateDirectoryTree(*InName);
-		FileName = FPaths::Combine(InName, FPaths::GetCleanFilename(InName));;
-	}
-	else
-	{
-		FileName = InName;
-	}
-
-	// Loop over y tiles.
-	for (int TileY = 0; TileY < NumTilesY; ++TileY)
-	{
-		// Loop over x tiles.
-		for (int TileX = 0; TileX < NumTilesX; ++TileX)
-		{
-			// Copy tile line by line.
-			uint8* DestPtr = TileRawData.GetData();
-			uint8* SrcPtr = RawData.GetData() + TileX * TileWidth * BytesPerPixel +
-				TileY * TileHeight * Width * BytesPerPixel;
-			for (int LineY = 0; LineY < TileHeight; ++LineY)
-			{
-				FMemory::Memcpy(DestPtr, SrcPtr,
-					TileWidth * BytesPerPixel);
-				DestPtr += TileWidth * BytesPerPixel;
-				SrcPtr += Width * BytesPerPixel;
-			}
-
-			// Compress data.
-			InImageWrapper->SetRaw(TileRawData.GetData(), TileRawData.Num(),
-				TileWidth, TileHeight, Format, BitDepth);
-			const TArray64<uint8> CompressedData = InImageWrapper->GetCompressed((int32)EImageCompressionQuality::Uncompressed);
-
-			// Write out tile.
-			FString Name = FString::Format(TEXT("{0}_x{1}_y{2}.{3}"),
-				{*FileName, TileX, TileY, *FileExtension});
-			FFileHelper::SaveArrayToFile(CompressedData, *Name);
-		}
-	}
-}
-
 void SImgMediaProcessImages::ProcessImageCustom(TSharedPtr<IImageWrapper>& InImageWrapper,
 	int32 InTileWidth, int32 InTileHeight, int32 InTileBorder, bool bInEnableMips,
-	bool bHasAlphaChannel, const FString& InName)
+	bool bHasAlphaChannel, const FString& InName,
+	bool bIsCustomFormat)
 {
 #if IMGMEDIAEDITOR_EXR_SUPPORTED_PLATFORM
 	TRACE_CPUPROFILER_EVENT_SCOPE(SImgMediaProcessImages::ProcessImageCustom);
@@ -436,7 +365,7 @@ void SImgMediaProcessImages::ProcessImageCustom(TSharedPtr<IImageWrapper>& InIma
 	}
 	ProcessImageCustomRawData(RawData, Width, Height, BitDepth,
 		InTileWidth, InTileHeight, InTileBorder, bInEnableMips,
-		bHasAlphaChannel, InName);
+		bHasAlphaChannel, InName, bIsCustomFormat);
 #else // IMGMEDIAEDITOR_EXR_SUPPORTED_PLATFORM
 	UE_LOG(LogImgMediaEditor, Error, TEXT("EXR not supported on this platform."));
 #endif // IMGMEDIAEDITOR_EXR_SUPPORTED_PLATFORM
@@ -445,7 +374,8 @@ void SImgMediaProcessImages::ProcessImageCustom(TSharedPtr<IImageWrapper>& InIma
 void SImgMediaProcessImages::ProcessImageCustomRawData(TArray64<uint8>& RawData,
 	int32 Width, int32 Height, int32 BitDepth,
 	int32 InTileWidth, int32 InTileHeight, int32 InTileBorder, bool bInEnableMips,
-	bool bHasAlphaChannel, const FString& InName)
+	bool bHasAlphaChannel, const FString& InName,
+	bool bIsCustomFormat)
 {
 #if IMGMEDIAEDITOR_EXR_SUPPORTED_PLATFORM
 	TRACE_CPUPROFILER_EVENT_SCOPE(SImgMediaProcessImages::ProcessImageCustomRawData);
@@ -491,17 +421,20 @@ void SImgMediaProcessImages::ProcessImageCustomRawData(TArray64<uint8>& RawData,
 
 	// Create tiled exr file.
 	FTiledOutputFile OutFile(FIntPoint(0, 0), FIntPoint(DestWidth - 1, DestHeight - 1),
-		FIntPoint(0, 0), FIntPoint(DestWidth - 1, DestHeight - 1));
+		FIntPoint(0, 0), FIntPoint(DestWidth - 1, DestHeight - 1), bIsTiled || bIsCustomFormat);
 
 	// Add attributes.
-	OutFile.AddIntAttribute(IImgMediaModule::CustomFormatAttributeName.Resolve().ToString(), 1);
-
-	// These attributes will not be added and therefore not found by EXR reader if it is not tiled.
-	if (bIsTiled)
+	if (bIsCustomFormat)
 	{
-		OutFile.AddIntAttribute(IImgMediaModule::CustomFormatTileWidthAttributeName.Resolve().ToString(), TileWidth);
-		OutFile.AddIntAttribute(IImgMediaModule::CustomFormatTileHeightAttributeName.Resolve().ToString(),TileHeight);
-		OutFile.AddIntAttribute(IImgMediaModule::CustomFormatTileBorderAttributeName.Resolve().ToString(), InTileBorder);
+		OutFile.AddIntAttribute(IImgMediaModule::CustomFormatAttributeName.Resolve().ToString(), 1);
+
+		// These attributes will not be added and therefore not found by EXR reader if it is not tiled.
+		if (bIsTiled)
+		{
+			OutFile.AddIntAttribute(IImgMediaModule::CustomFormatTileWidthAttributeName.Resolve().ToString(), TileWidth);
+			OutFile.AddIntAttribute(IImgMediaModule::CustomFormatTileHeightAttributeName.Resolve().ToString(), TileHeight);
+			OutFile.AddIntAttribute(IImgMediaModule::CustomFormatTileBorderAttributeName.Resolve().ToString(), InTileBorder);
+		}
 	}
 
 	// Add channels.
@@ -517,7 +450,15 @@ void SImgMediaProcessImages::ProcessImageCustomRawData(TArray64<uint8>& RawData,
 	}
 
 	// Create output.
-	OutFile.CreateOutputFile(InName, DestWidth, DestHeight, bInEnableMips, 1);
+	if (bIsCustomFormat)
+	{
+		OutFile.CreateOutputFile(InName, DestWidth, DestHeight, bInEnableMips, 1);
+	}
+	else
+	{
+		OutFile.CreateOutputFile(InName, TileWidth, TileHeight, bInEnableMips, 1);
+	}
+
 	if (DestNumChannels == 4)
 	{
 		OutFile.AddFrameBufferChannel(AChannelName, nullptr, Stride);
@@ -603,7 +544,7 @@ void SImgMediaProcessImages::ProcessImageCustomRawData(TArray64<uint8>& RawData,
 
 		// Do we need to tile this mip?
 		// Need to also check that this is actually a valid mip level.
-		if ((bIsTiled) && (MipSourceWidth > 0) && (MipSourceHeight > 0))
+		if ((bIsCustomFormat) && (bIsTiled) && (MipSourceWidth > 0) && (MipSourceHeight > 0))
 		{
 			int32 MipTileWidth = TileWidth;
 			int32 MipTileHeight = TileHeight;
@@ -652,25 +593,55 @@ void SImgMediaProcessImages::ProcessImageCustomRawData(TArray64<uint8>& RawData,
 
 		// Write to EXR.
 		TRACE_CPUPROFILER_EVENT_SCOPE(SImgMediaProcessImages::ProcessImageCustom:WriteEXR);
-		Stride.Y = MipWidth * BytesPerPixel;
-		int64 BufferOffset = 0;
-		int64 SingleBufferOffset = MipWidth * BytesPerPixelPerChannel;
-		if (DestNumChannels == 4)
+		if (bIsCustomFormat)
 		{
-			OutFile.UpdateFrameBufferChannel(AChannelName, CurrentBuffer, Stride);
+			Stride.Y = MipWidth * BytesPerPixel;
+			int64 BufferOffset = 0;
+			int64 SingleBufferOffset = MipWidth * BytesPerPixelPerChannel;
+			if (DestNumChannels == 4)
+			{
+				OutFile.UpdateFrameBufferChannel(AChannelName, CurrentBuffer, Stride);
+				BufferOffset += SingleBufferOffset;
+			}
+
+			OutFile.UpdateFrameBufferChannel(BChannelName, CurrentBuffer + BufferOffset, Stride);
+			BufferOffset += SingleBufferOffset;
+			OutFile.UpdateFrameBufferChannel(GChannelName, CurrentBuffer + BufferOffset, Stride);
+			BufferOffset += SingleBufferOffset;
+			OutFile.UpdateFrameBufferChannel(RChannelName, CurrentBuffer + BufferOffset, Stride);
 			BufferOffset += SingleBufferOffset;
 		}
+		else
+		{
+			Stride.X = BytesPerPixel;
+			Stride.Y = MipWidth * BytesPerPixel;
+			int64 BufferOffset = 0;
+			int64 SingleBufferOffset = BytesPerPixelPerChannel;
+			
+
+			OutFile.UpdateFrameBufferChannel(RChannelName, CurrentBuffer + BufferOffset, Stride);
+			BufferOffset += SingleBufferOffset;
+			OutFile.UpdateFrameBufferChannel(GChannelName, CurrentBuffer + BufferOffset, Stride);
+			BufferOffset += SingleBufferOffset;
+			OutFile.UpdateFrameBufferChannel(BChannelName, CurrentBuffer + BufferOffset, Stride);
+			BufferOffset += SingleBufferOffset;
 
-		OutFile.UpdateFrameBufferChannel(BChannelName, CurrentBuffer + BufferOffset, Stride);
-		BufferOffset += SingleBufferOffset;
-		OutFile.UpdateFrameBufferChannel(GChannelName, CurrentBuffer + BufferOffset, Stride);
-		BufferOffset += SingleBufferOffset;
-		OutFile.UpdateFrameBufferChannel(RChannelName, CurrentBuffer + BufferOffset, Stride);
-		BufferOffset += SingleBufferOffset;
+			if (DestNumChannels == 4)
+			{
+				OutFile.UpdateFrameBufferChannel(AChannelName, CurrentBuffer + BufferOffset, Stride);
+			}
+		}
 
 		OutFile.SetFrameBuffer();
 
-		OutFile.WriteTile(0, 0, MipLevel);
+		if (bIsCustomFormat)
+		{
+			OutFile.WriteTile(0, 0, MipLevel);
+		}
+		else
+		{
+			OutFile.WriteTiles(0, OutFile.GetNumXTiles(MipLevel) - 1, 0, OutFile.GetNumYTiles(MipLevel) - 1, MipLevel);
+		}
 
 		// Switch buffers.
 		CurrentMipBufferIndex ^= 1;
@@ -929,11 +900,11 @@ void SImgMediaProcessImages::HandleProcessing()
 
 						Async(EAsyncExecution::Thread, [this, RawData = MoveTemp(RawData), Width, Height, BitDepth,
 							InTileWidth, InTileHeight, TileBorder, bEnableMips,
-							bHasAlphaChannel, Name]() mutable
+							bHasAlphaChannel, Name, bUseCustomFormat]() mutable
 						{
 							ProcessImageCustomRawData(RawData, Width, Height, BitDepth,
 								InTileWidth, InTileHeight, TileBorder, bEnableMips,
-								bHasAlphaChannel, Name);
+								bHasAlphaChannel, Name, bUseCustomFormat);
 						});
 					}
 					else
diff --git a/Engine/Plugins/Media/ImgMedia/Source/ImgMediaEditor/Private/Widgets/SImgMediaProcessImages.h b/Engine/Plugins/Media/ImgMedia/Source/ImgMediaEditor/Private/Widgets/SImgMediaProcessImages.h
--- a/Engine/Plugins/Media/ImgMedia/Source/ImgMediaEditor/Private/Widgets/SImgMediaProcessImages.h
+++ b/Engine/Plugins/Media/ImgMedia/Source/ImgMediaEditor/Private/Widgets/SImgMediaProcessImages.h
@@ -59,20 +59,6 @@ private:
 	  */
 	bool HasAlphaChannel(const FString& Ext, const FString& File);
 
-	/**
-	 * Processess a single image and writes out 1 or more files.
-	 * Tiles and mips may be generated.
-	 * This does NOT run on the game thread.
-	 *
-	 * @param InImageWrapper	ImageWrapper to read/write the image.
-	 * @param InTileWidth		Desired width of tiles.
-	 * @param InTileHeight		Desired height of tiles.
-	 * @param InName			Full path and name of file to write to WITHOUT the extension (e.g. no .exr)
-	 * @param FileExtension		Extension to append to the file name.
-	 */
-	void ProcessImage(TSharedPtr<IImageWrapper>& InImageWrapper,
-		int32 InTileWidth, int32 InTileHeight, const FString& InName, const FString& FileExtension);
-
 	/**
 	 * Processess a single image and writes out a file.
 	 * Tiles and mips may be generated.
@@ -85,10 +71,12 @@ private:
 	 * @param bInEnableMips		Turn on mip mapping.
 	 * @param bHasAlphaChannel	True if there really is an alpha channel.
 	 * @param InName			Full path and name of file to write.
+	 * @param bIsCustomFormat	True to output in our custom format.
 	 */
 	void ProcessImageCustom(TSharedPtr<IImageWrapper>& InImageWrapper,
 		int32 InTileWidth, int32 InTileHeight, int32 InTileBorder, 
-		bool bInEnableMips, bool bHasAlphaChannel, const FString& InName);
+		bool bInEnableMips, bool bHasAlphaChannel, const FString& InName,
+		bool bIsCustomFormat);
 
 	/**
 	 * Processess a single image and writes out a file.
@@ -105,11 +93,12 @@ private:
 	 * @param bInEnableMips		Turn on mip mapping.
 	 * @param bHasAlphaChannel	True if there really is an alpha channel.
 	 * @param InName			Full path and name of file to write.
+	 * @param bIsCustomFormat	True to output in our custom format.
 	 */
 	void ProcessImageCustomRawData(TArray64<uint8>& RawData,
 		int32 Width, int32 Height, int32 BitDepth,
 		int32 InTileWidth, int32 InTileHeight, int32 InTileBorder, bool bInEnableMips,
-		bool bHasAlphaChannel, const FString& InName);
+		bool bHasAlphaChannel, const FString& InName, bool bIsCustomFormat);
 	
 	/**
 	 * Removes the alpha channel from a buffer.
diff --git a/Engine/Plugins/Media/ImgMedia/Source/OpenExrWrapper/Private/OpenExrWrapper.cpp b/Engine/Plugins/Media/ImgMedia/Source/OpenExrWrapper/Private/OpenExrWrapper.cpp
--- a/Engine/Plugins/Media/ImgMedia/Source/OpenExrWrapper/Private/OpenExrWrapper.cpp
+++ b/Engine/Plugins/Media/ImgMedia/Source/OpenExrWrapper/Private/OpenExrWrapper.cpp
@@ -14,6 +14,7 @@ THIRD_PARTY_INCLUDES_START
 	#include "OpenEXR/ImfCompressionAttribute.h"
 	#include "OpenEXR/ImfHeader.h"
 	#include "OpenEXR/ImfIntAttribute.h"
+	#include "OpenEXR/ImfOutputFile.h"
 	#include "OpenEXR/ImfTileDescriptionAttribute.h"
 	#include "OpenEXR/ImfRgbaFile.h"
 	#include "OpenEXR/ImfStandardAttributes.h"
@@ -407,8 +408,10 @@ FTiledOutputFile::FTiledOutputFile(
 	const FIntPoint& DisplayWindowMin,
 	const FIntPoint& DisplayWindowMax,
 	const FIntPoint& DataWindowMin,
-	const FIntPoint& DataWindowMax)
+	const FIntPoint& DataWindowMax,
+	bool bInIsTiled)
 	: FBaseOutputFile(DisplayWindowMin, DisplayWindowMax, DataWindowMin, DataWindowMax)
+	, bIsTiled(bInIsTiled)
 {
 	FrameBuffer = new Imf::FrameBuffer;
 }
@@ -433,13 +436,21 @@ void FTiledOutputFile::CreateOutputFile(const FString& FilePath,
 	{
 		try
 		{
-			((Imf::Header*)Header)->setTileDescription(Imf::TileDescription(TileWidth, TileHeight,
-				bIsMipsEnabled ? Imf::MIPMAP_LEVELS : Imf::ONE_LEVEL));
-
+			if (bIsTiled)
+			{
+				((Imf::Header*)Header)->setTileDescription(Imf::TileDescription(TileWidth, TileHeight,
+					bIsMipsEnabled ? Imf::MIPMAP_LEVELS : Imf::ONE_LEVEL));
 
-			// Create output file.
-			OutputFile = new Imf::TiledOutputFile(StringCast<ANSICHAR>(*FilePath).Get(),
-				*((Imf::Header*)Header), NumThreads);
+				// Create output file.
+				OutputFile = new Imf::TiledOutputFile(StringCast<ANSICHAR>(*FilePath).Get(),
+					*((Imf::Header*)Header), NumThreads);
+			}
+			else
+			{
+				// Create output file.
+				OutputFile = new Imf::OutputFile(StringCast<ANSICHAR>(*FilePath).Get(),
+					*((Imf::Header*)Header), NumThreads);
+			}
 		}
 		catch (std::exception const& Exception)
 		{
@@ -482,7 +493,14 @@ void FTiledOutputFile::SetFrameBuffer()
 {
 	if (OutputFile != nullptr)
 	{
-		((Imf::TiledOutputFile*)OutputFile)->setFrameBuffer(*((Imf::FrameBuffer*)FrameBuffer));
+		if (bIsTiled)
+		{
+			((Imf::TiledOutputFile*)OutputFile)->setFrameBuffer(*((Imf::FrameBuffer*)FrameBuffer));
+		}
+		else
+		{
+			((Imf::OutputFile*)OutputFile)->setFrameBuffer(*((Imf::FrameBuffer*)FrameBuffer));
+		}
 	}
 	else
 	{
@@ -492,12 +510,48 @@ void FTiledOutputFile::SetFrameBuffer()
 }
 
 void FTiledOutputFile::WriteTile(int32 TileX, int32 TileY, int32 MipLevel)
+{
+	if (bIsTiled)
+	{
+		if (OutputFile != nullptr)
+		{
+			try
+			{
+				((Imf::TiledOutputFile*)OutputFile)->writeTiles(0, TileX, 0, TileY, MipLevel);
+			}
+			catch (std::exception const& Exception)
+			{
+				UE_LOG(LogOpenEXRWrapper, Error, TEXT("Cannot write EXR file: %s"),
+					StringCast<TCHAR>(Exception.what()).Get());
+			}
+		}
+		else
+		{
+			UE_LOG(LogOpenEXRWrapper, Error,
+				TEXT("WriteTile failed: CreateOutputFile has not been called yet."));
+		}
+	}
+	else
+	{
+		WriteTiles(0, 0, 0, 0, 0);
+	}
+}
+
+void FTiledOutputFile::WriteTiles(int32 TileX1, int32 TileX2, int32 TileY1, int32 TileY2, int32 MipLevel)
 {
 	if (OutputFile != nullptr)
 	{
 		try
 		{
-			((Imf::TiledOutputFile*)OutputFile)->writeTile(TileX, TileY, MipLevel);
+			if (bIsTiled)
+			{
+				((Imf::TiledOutputFile*)OutputFile)->writeTiles(TileX1, TileX2, TileY1, TileY2, MipLevel);
+			}
+			else
+			{
+				Imath::Box2i DataWindow = ((Imf::Header*)Header)->dataWindow();
+				((Imf::OutputFile*)OutputFile)->writePixels(DataWindow.max.y - DataWindow.min.y + 1);
+			}
 		}
 		catch (std::exception const& Exception)
 		{
@@ -508,7 +562,7 @@ void FTiledOutputFile::WriteTile(int32 TileX, int32 TileY, int32 MipLevel)
 	else
 	{
 		UE_LOG(LogOpenEXRWrapper, Error,
-			TEXT("WriteTile failed: CreateOutputFile has not been called yet."));
+			TEXT("WriteTiles failed: CreateOutputFile has not been called yet."));
 	}
 }
 
@@ -516,7 +570,14 @@ int32 FTiledOutputFile::GetNumberOfMipLevels()
 {
 	if (OutputFile != nullptr)
 	{
-		return ((Imf::TiledOutputFile*)OutputFile)->numLevels();
+		if (bIsTiled)
+		{
+			return ((Imf::TiledOutputFile*)OutputFile)->numLevels();
+		}
+		else
+		{
+			return 1;
+		}
 	}
 	else
 	{
@@ -530,7 +591,15 @@ int32 FTiledOutputFile::GetMipWidth(int32 MipLevel)
 {
 	if (OutputFile != nullptr)
 	{
-		return ((Imf::TiledOutputFile*)OutputFile)->levelWidth(MipLevel);
+		if (bIsTiled)
+		{
+			return ((Imf::TiledOutputFile*)OutputFile)->levelWidth(MipLevel);
+		}
+		else
+		{
+			Imath::Box2i DataWindow = ((Imf::Header*)Header)->dataWindow();
+			return DataWindow.size().x + 1;
+		}
 	}
 	else
 	{
@@ -544,7 +613,15 @@ int32 FTiledOutputFile::GetMipHeight(int32 MipLevel)
 {
 	if (OutputFile != nullptr)
 	{
-		return ((Imf::TiledOutputFile*)OutputFile)->levelHeight(MipLevel);
+		if (bIsTiled)
+		{
+			return ((Imf::TiledOutputFile*)OutputFile)->levelHeight(MipLevel);
+		}
+		else
+		{
+			Imath::Box2i DataWindow = ((Imf::Header*)Header)->dataWindow();
+			return DataWindow.size().y + 1;
+		}
 	}
 	else
 	{
@@ -554,4 +631,46 @@ int32 FTiledOutputFile::GetMipHeight(int32 MipLevel)
 	}
 }
 
+int32 FTiledOutputFile::GetNumXTiles(int32 MipLevel)
+{
+	if (OutputFile != nullptr)
+	{
+		if (bIsTiled)
+		{
+			return ((Imf::TiledOutputFile*)OutputFile)->numXTiles(MipLevel);
+		}
+		else
+		{
+			return 0;
+		}
+	}
+	else
+	{
+		UE_LOG(LogOpenEXRWrapper, Error,
+			TEXT("GetNumXTiles failed: CreateOutputFile has not been called yet."));
+		return 0;
+	}
+}
+
+int32 FTiledOutputFile::GetNumYTiles(int32 MipLevel)
+{
+	if (OutputFile != nullptr)
+	{
+		if (bIsTiled)
+		{
+			return ((Imf::TiledOutputFile*)OutputFile)->numYTiles(MipLevel);
+		}
+		else
+		{
+			return 0;
+		}
+	}
+	else
+	{
+		UE_LOG(LogOpenEXRWrapper, Error,
+			TEXT("GetNumYTiles failed: CreateOutputFile has not been called yet."));
+		return 0;
+	}
+}
+
 IMPLEMENT_MODULE(FDefaultModuleImpl, OpenExrWrapper);
diff --git a/Engine/Plugins/Media/ImgMedia/Source/OpenExrWrapper/Public/OpenExrWrapper.h b/Engine/Plugins/Media/ImgMedia/Source/OpenExrWrapper/Public/OpenExrWrapper.h
--- a/Engine/Plugins/Media/ImgMedia/Source/OpenExrWrapper/Public/OpenExrWrapper.h
+++ b/Engine/Plugins/Media/ImgMedia/Source/OpenExrWrapper/Public/OpenExrWrapper.h
@@ -213,12 +213,14 @@ public:
 	 * @param DisplayWindowMax		Normally (width - 1, height - 1).
 	 * @param DataWindowMin			Normally (0, 0).
 	 * @param DataWindowMax			Normally (width - 1, height - 1).
+	 * @param bInIsTiled			True if you want tiles.
 	 */
 	FTiledOutputFile(
 		const FIntPoint& DisplayWindowMin,
 		const FIntPoint& DisplayWindowMax,
 		const FIntPoint& DataWindowMin,
-		const FIntPoint& DataWindowMax);
+		const FIntPoint& DataWindowMax,
+		bool bInIsTiled = true);
 
 	/**
 	 * Destructor.
@@ -254,6 +256,16 @@ public:
 	 */
 	int32 GetMipHeight(int32 MipLevel);
 
+	/**
+	 * Get number of horizontal tiles for a mip level.
+	 */
+	int32 GetNumXTiles(int32 MipLevel);
+
+	/**
+	 * Get number of vertical tiles for a mip level.
+	 */
+	int32 GetNumYTiles(int32 MipLevel);
+
 	/**
 	 * Call this before SetFrameBuffer for each channel in the frame buffer.
 	 *
@@ -284,6 +296,7 @@ public:
 
 	/**
 	 * Call this to write data to a specific tile.
+	 * If tiles are not enabled then this just writes out the whole image.
 	 *
 	 * @param TileX			X coordinate of tile.
 	 * @param TileY			Y coordinate of tile.
@@ -291,6 +304,18 @@ public:
 	 */
 	void WriteTile(int32 TileX, int32 TileY, int32 MipLevel);
 
+	/**
+	 * Call this to write data to some tiles.
+	 * If tiles are not enabled then this just writes out the whole image.
+	 *
+	 * @param TileX1		X coordinate of first tile.
+	 * @param TileX2		X coordinate of last tile.
+	 * @param TileY1		Y coordinate of first tile.
+	 * @param TileY2		Y coordinate of last tile.
+	 * @param MipLevel		Mipmap level of tile.
+	 */
+	void WriteTiles(int32 TileX1, int32 TileX2, int32 TileY1, int32 TileY2, int32 MipLevel);
+
 	//~ FBaseOutputFile interface.
 	virtual int32 GetNumberOfMipLevels() override;
 
@@ -298,5 +323,7 @@ private:
 
 	/** Stores the EXR frame buffer object. */
 	void* FrameBuffer;
+	/** True if tiles are enabled. */
+	bool bIsTiled;
 };
 
