ImgMedia: Fixed ImgMediaSource not calling Super::PostEditChangeProperty.
#rnx

diff --git a/Engine/Plugins/Media/ImgMedia/Source/ImgMedia/Private/Assets/ImgMediaSource.cpp b/Engine/Plugins/Media/ImgMedia/Source/ImgMedia/Private/Assets/ImgMediaSource.cpp
--- a/Engine/Plugins/Media/ImgMedia/Source/ImgMedia/Private/Assets/ImgMediaSource.cpp
+++ b/Engine/Plugins/Media/ImgMedia/Source/ImgMedia/Private/Assets/ImgMediaSource.cpp
@@ -214,6 +214,8 @@ void UImgMediaSource::PostEditChangeProperty(FPropertyChangedEvent& PropertyChan
 			GlobalCache->EmptyCache();
 		}
 	}
+
+	Super::PostEditChangeProperty(PropertyChangedEvent);
 }
 
 #endif // WITH_EDITOR
