ProcessImages: Added support for 32bit EXRs.

#rnx

diff --git a/Engine/Plugins/Media/ImgMedia/Source/ImgMediaEditor/Private/Widgets/SImgMediaProcessImages.cpp b/Engine/Plugins/Media/ImgMedia/Source/ImgMediaEditor/Private/Widgets/SImgMediaProcessImages.cpp
--- a/Engine/Plugins/Media/ImgMedia/Source/ImgMediaEditor/Private/Widgets/SImgMediaProcessImages.cpp
+++ b/Engine/Plugins/Media/ImgMedia/Source/ImgMediaEditor/Private/Widgets/SImgMediaProcessImages.cpp
@@ -389,6 +389,20 @@ void SImgMediaProcessImages::ProcessImageCustomRawData(TArray64<uint8>& RawData,
 	int32 BytesPerPixelPerChannel = BitDepth / 8;
 	int32 NumChannels = BytesPerPixel / BytesPerPixelPerChannel;
 	int32 DestNumChannels = NumChannels;
+
+	// We use 16 bits only, so convert if we have 32.
+	if (BitDepth == 32)
+	{
+		ConvertTo16Bit(RawData);
+		BitDepth = 16;
+		BytesPerPixel /= 2;
+		BytesPerPixelPerChannel /= 2;
+	}
+	else if (BitDepth != 16)
+	{
+		UE_LOG(LogImgMediaEditor, Error, TEXT("Unsupported bit depth %d"), BitDepth);
+	}
+
 	// ImageWrapper always returns an alpha channel, so make sure we really have one.
 	if ((DestNumChannels == 4) && (bHasAlphaChannel == false))
 	{
@@ -656,6 +670,27 @@ void SImgMediaProcessImages::ProcessImageCustomRawData(TArray64<uint8>& RawData,
 #endif // IMGMEDIAEDITOR_EXR_SUPPORTED_PLATFORM
 }
 
+
+void SImgMediaProcessImages::ConvertTo16Bit(TArray64<uint8>& Buffer)
+{
+	TRACE_CPUPROFILER_EVENT_SCOPE(SImgMediaProcessImages::ConvertTo16Bit);
+
+	int32 BytesPerPixelPerChannel = 4;
+	int64 BufferSize = Buffer.Num() / BytesPerPixelPerChannel;
+	FFloat32* InBufferPtr = (FFloat32*)(Buffer.GetData());
+	FFloat16* OutBufferPtr = (FFloat16*)(Buffer.GetData());
+
+	// Loop through the buffer.
+	for (int64 Index = 0; Index < BufferSize; ++Index)
+	{
+		// Convert the data in place.
+		OutBufferPtr[Index] = InBufferPtr[Index].FloatValue;
+	}
+
+	// Don't bother shrinking as its just a waste and extra work.
+	Buffer.SetNum(BufferSize / 2, false);
+}
+
 void SImgMediaProcessImages::RemoveAlphaChannel(TArray64<uint8>& Buffer)
 {
 	TRACE_CPUPROFILER_EVENT_SCOPE(SImgMediaProcessImages::RemoveAlphaChannel);
diff --git a/Engine/Plugins/Media/ImgMedia/Source/ImgMediaEditor/Private/Widgets/SImgMediaProcessImages.h b/Engine/Plugins/Media/ImgMedia/Source/ImgMediaEditor/Private/Widgets/SImgMediaProcessImages.h
--- a/Engine/Plugins/Media/ImgMedia/Source/ImgMediaEditor/Private/Widgets/SImgMediaProcessImages.h
+++ b/Engine/Plugins/Media/ImgMedia/Source/ImgMediaEditor/Private/Widgets/SImgMediaProcessImages.h
@@ -100,6 +100,14 @@ private:
 		int32 InTileWidth, int32 InTileHeight, int32 InTileBorder, bool bInEnableMips,
 		bool bHasAlphaChannel, const FString& InName, bool bIsCustomFormat);
 	
+	/**
+	 * Converts a buffer from 32 bit floating point to 16 bit floating point.
+	 * This will be done in place.
+	 * 
+	 * @param Buffer			Buffer to convert.
+	 */
+	void ConvertTo16Bit(TArray64<uint8>& Buffer);
+
 	/**
 	 * Removes the alpha channel from a buffer.
 	 * This will be done in place.
