ImgMedia: Improvements to bandwidth calculations for tiled image sequences.
Added support for tiling, added effective bandwidth to show the fraction of the current (available) disk bandwidth is in use.

diff --git a/Engine/Plugins/Media/ImgMedia/Source/ImgMedia/Private/Loader/ImgMediaLoader.cpp b/Engine/Plugins/Media/ImgMedia/Source/ImgMedia/Private/Loader/ImgMediaLoader.cpp
--- a/Engine/Plugins/Media/ImgMedia/Source/ImgMedia/Private/Loader/ImgMediaLoader.cpp
+++ b/Engine/Plugins/Media/ImgMedia/Source/ImgMedia/Private/Loader/ImgMediaLoader.cpp
@@ -79,6 +79,74 @@ namespace ImgMediaLoader
 	}
 }
 
+FImgMediaLoaderBandwidth::FImgMediaLoaderBandwidth()
+	: Current(0)
+	, Effective(0)
+	, Required(0)
+	, ReadTimeCache()
+	, ReadTimeCacheIndex(0)
+{
+
+}
+
+void FImgMediaLoaderBandwidth::Update(const TSharedPtr<FImgMediaFrame, ESPMode::ThreadSafe>& Frame, double WorkTime)
+{
+	static constexpr int32 READ_TIME_CACHE_MAX = 64;
+
+	// Calculate the current uncompressed bandwidth
+	SIZE_T BytesLoaded = Frame->Info.UncompressedSize;
+	Current = (float)BytesLoaded / WorkTime;
+
+	if (Frame->Info.bHasTiles)
+	{
+		int32 TotalNumTiles = 0;
+
+		for (int32 MipLevel = 0; MipLevel < Frame->Info.NumMipLevels; ++MipLevel)
+		{
+			const int32 MipLevelDiv = 1 << MipLevel;
+			int32 NumTilesX = FMath::Max(1, FMath::CeilToInt(float(Frame->Info.NumTiles.X) / MipLevelDiv));
+			int32 NumTilesY = FMath::Max(1, FMath::CeilToInt(float(Frame->Info.NumTiles.Y) / MipLevelDiv));
+			TotalNumTiles += NumTilesX * NumTilesY;
+		}
+
+		// Adjusted based on how many tiles were loaded
+		Current *= (float)Frame->NumTilesRead / FMath::Max(1, TotalNumTiles);
+	}
+
+	// Store read times with their respectivate timestamps...
+	if (ReadTimeCache.Num() == READ_TIME_CACHE_MAX)
+	{
+		ReadTimeCache[ReadTimeCacheIndex] = TPairInitializer(FPlatformTime::Seconds(), WorkTime);
+		ReadTimeCacheIndex = (ReadTimeCacheIndex + 1) % READ_TIME_CACHE_MAX;
+	}
+	else
+	{
+		ReadTimeCache.Emplace(FPlatformTime::Seconds(), WorkTime);
+	}
+
+	// Calculate the effective/utilized bandwidth as the fraction of the elapsed time spent reading.
+	double IntervalStart = TNumericLimits<double>::Max();
+	double IntervalEnd = TNumericLimits<double>::Lowest();
+	double TotalReadTime = 0.0;
+
+	for (const TPair<double, double>& Pair : ReadTimeCache)
+	{
+		double CachedTimestamp = Pair.Key;
+		double CachedReadTime = Pair.Value;
+		IntervalStart = FMath::Min(IntervalStart, CachedTimestamp - CachedReadTime);
+		IntervalEnd = FMath::Max(IntervalEnd, CachedTimestamp);
+		TotalReadTime += CachedReadTime;
+	}
+
+	Effective = (TotalReadTime / (IntervalEnd - IntervalStart)) * Current;
+}
+
+void FImgMediaLoaderBandwidth::EmptyCache()
+{
+	ReadTimeCache.Empty(ReadTimeCache.Max());
+	ReadTimeCacheIndex = 0;
+}
+
 /* FImgMediaLoader structors
  *****************************************************************************/
 
@@ -104,6 +172,9 @@ FImgMediaLoader::FImgMediaLoader(const TSharedRef<FImgMediaScheduler, ESPMode::T
 	, RetryCount(0)
 	, UseGlobalCache(false)
 	, SmartCacheSettings(InSmartCacheSettings)
+#if WITH_EDITOR
+	, Bandwidth()
+#endif
 {
 	ResetFetchLogic();
 	UE_LOG(LogImgMedia, Verbose, TEXT("Loader %p: Created"), this);
@@ -234,6 +305,12 @@ void FImgMediaLoader::ResetFetchLogic()
 	QueuedSampleFetch.CurrentSequenceIndex = 0;
 }
 
+void FImgMediaLoader::HandlePause()
+{
+#if WITH_EDITOR
+	Bandwidth.EmptyCache();
+#endif
+}
 
 float FImgMediaLoader::FindMaxOverlapInRange(int32 StartIndex, int32 EndIndex, FTimespan StartTime, FTimespan EndTime, int32 & MaxIdx) const
 {
@@ -836,14 +913,6 @@ void FImgMediaLoader::LoadSequence(const FString& SequencePath, const FFrameRate
 	SequenceDuration = FrameNumberToTime(GetNumImages());
 	SIZE_T UncompressedSize = FirstFrameInfo.UncompressedSize;
 
-	// Get minimum bandwidth needed for mip level 0.
-	MinBandwidthForMipLevel0 = UncompressedSize;
-	float FrameTime = SequenceFrameRate.AsInterval();
-	if (FrameTime > 0)
-	{
-		MinBandwidthForMipLevel0 /= FrameTime;
-	}
-
 	if (FirstFrameInfo.bHasTiles)
 	{
 		TilingDescription.TileBorderSize = FirstFrameInfo.TileBorder;
@@ -851,6 +920,15 @@ void FImgMediaLoader::LoadSequence(const FString& SequencePath, const FFrameRate
 		TilingDescription.TileSize = FirstFrameInfo.TileDimensions;
 	}
 
+#if WITH_EDITOR
+	// Get required bandwidth needed for all mips and tiles.
+	float FrameTime = SequenceFrameRate.AsInterval();
+	if (FrameTime > 0)
+	{
+		Bandwidth.Required = UncompressedSize / FrameTime;
+	}
+#endif
+
 	// If we have no mips or tiles, then get rid of our MipMapInfoObject.
 	// Otherwise, set it up.
 	if (MipMapInfo.IsValid())
@@ -1275,12 +1353,13 @@ void FImgMediaLoader::NotifyWorkComplete(FImgMediaLoaderWork& CompletedWork, int
 
 	WorkPool.Push(&CompletedWork);
 
+#if WITH_EDITOR
 	// Update bandwidth.
-	if ((Frame.IsValid()) && (WorkTime > 0.0f))
+	if (Frame.IsValid() && WorkTime > 0.0)
 	{
-		SIZE_T BytesLoaded = Frame->Info.UncompressedSize;
-		CurrentBandwidth = BytesLoaded / WorkTime;
+		Bandwidth.Update(Frame, WorkTime);
 	}
+#endif
 }
 
 void FImgMediaLoader::AddFrameToCache(int32 FrameNumber, const TSharedPtr<FImgMediaFrame, ESPMode::ThreadSafe>& Frame)
diff --git a/Engine/Plugins/Media/ImgMedia/Source/ImgMedia/Private/Loader/ImgMediaLoader.h b/Engine/Plugins/Media/ImgMedia/Source/ImgMedia/Private/Loader/ImgMediaLoader.h
--- a/Engine/Plugins/Media/ImgMedia/Source/ImgMedia/Private/Loader/ImgMediaLoader.h
+++ b/Engine/Plugins/Media/ImgMedia/Source/ImgMedia/Private/Loader/ImgMediaLoader.h
@@ -41,6 +41,34 @@ struct FImgMediaLoaderSmartCacheSettings
 	}
 };
 
+/**
+ * Sequence playback bandwidth requirements and estimates.
+ */
+struct FImgMediaLoaderBandwidth
+{
+	/** Number of bytes read per second from the last frame read. */
+	float Current;
+	/** Average bandwidth over time (including idle time). */
+	float Effective;
+	/** Minimum number of bytes per second needed to read all mip levels. */
+	float Required;
+
+	/** Constructor. */
+	FImgMediaLoaderBandwidth();
+
+	/** Updated calculated current and effective bandwidth numbers. */
+	void Update(const TSharedPtr<FImgMediaFrame, ESPMode::ThreadSafe>& Frame, double WorkTime);
+
+	/** Empty the read times cache and resets its index. */
+	void EmptyCache();
+
+private:
+	/** Work time cache for average bandwidth calculation. */
+	TArray<TPair<double, double>> ReadTimeCache;
+	/** Circular index for work time cache calculation. */
+	int32 ReadTimeCacheIndex;
+};
+
 /**
  * Loads image sequence frames from disk.
  */
@@ -74,10 +102,14 @@ public:
 	 */
 	uint64 GetBitRate() const;
 
+#if WITH_EDITOR
 	/** Get the number of bytes read per second from the last frame read. */
-	float GetCurrentBandwidth() const { return CurrentBandwidth; }
-	/** Get the minimum number of bytes per second needed to read mip level 0.. */
-	float GetMinBandwidthForMipLevel0() const { return MinBandwidthForMipLevel0; }
+	float GetCurrentBandwidth() const { return Bandwidth.Current; }
+	/** Get the effective bandwidth. */
+	float GetEffectiveBandwidth() const { return Bandwidth.Effective; }
+	/** Get the minimum number of bytes per second needed to read all mip levels. */
+	float GetRequiredBandwidth() const { return Bandwidth.Required; }
+#endif
 
 	/**
 	 * Get the time ranges of frames that are being loaded right now.
@@ -322,6 +354,11 @@ public:
 	 */
 	void ResetFetchLogic();
 
+	/**
+	 * Handler for when the (owning) player's state changes to EMediaState::Paused.
+	 */
+	void HandlePause();
+
 protected:
 
 	/**
@@ -556,8 +593,8 @@ private:
 	/** Settings for the smart cache. */
 	FImgMediaLoaderSmartCacheSettings SmartCacheSettings;
 
-	/** Number of bytes read per second from the last frame read. */
-	float CurrentBandwidth;
-	/** Minimum number of bytes per second needed to read mip level 0. */
-	float MinBandwidthForMipLevel0;
+#if WITH_EDITOR
+	/** Bandwidth estimation. */
+	FImgMediaLoaderBandwidth Bandwidth;
+#endif
 };
diff --git a/Engine/Plugins/Media/ImgMedia/Source/ImgMedia/Private/Loader/ImgMediaLoaderWork.cpp b/Engine/Plugins/Media/ImgMedia/Source/ImgMedia/Private/Loader/ImgMediaLoaderWork.cpp
--- a/Engine/Plugins/Media/ImgMedia/Source/ImgMedia/Private/Loader/ImgMediaLoaderWork.cpp
+++ b/Engine/Plugins/Media/ImgMedia/Source/ImgMedia/Private/Loader/ImgMediaLoaderWork.cpp
@@ -77,6 +77,7 @@ void FImgMediaLoaderWork::DoThreadedWork()
 		else
 		{
 			Frame = ExistingFrame;
+			Frame->NumTilesRead = 0;
 		}
 
 		double StartTime = FPlatformTime::Seconds();
diff --git a/Engine/Plugins/Media/ImgMedia/Source/ImgMedia/Private/Player/ImgMediaPlayer.cpp b/Engine/Plugins/Media/ImgMedia/Source/ImgMedia/Private/Player/ImgMediaPlayer.cpp
--- a/Engine/Plugins/Media/ImgMedia/Source/ImgMedia/Private/Player/ImgMediaPlayer.cpp
+++ b/Engine/Plugins/Media/ImgMedia/Source/ImgMedia/Private/Player/ImgMediaPlayer.cpp
@@ -600,6 +600,11 @@ bool FImgMediaPlayer::Seek(const FTimespan& Time)
 	if (CurrentState == EMediaState::Stopped)
 	{
 		CurrentState = EMediaState::Paused;
+
+		if (Loader.IsValid())
+		{
+			Loader->HandlePause();
+		}
 	}
 
 #if IMG_MEDIA_PLAYER_VERSION == 1
@@ -683,6 +688,11 @@ bool FImgMediaPlayer::SetRate(float Rate)
 		CurrentRate = Rate;
 		CurrentState = EMediaState::Paused;
 
+		if (Loader.IsValid())
+		{
+			Loader->HandlePause();
+		}
+
 		EventSink.ReceiveMediaEvent(EMediaEvent::PlaybackSuspended);
 		return true;
 	}
diff --git a/Engine/Plugins/Media/ImgMedia/Source/ImgMedia/Private/Readers/ExrImgMediaReader.cpp b/Engine/Plugins/Media/ImgMedia/Source/ImgMedia/Private/Readers/ExrImgMediaReader.cpp
--- a/Engine/Plugins/Media/ImgMedia/Source/ImgMedia/Private/Readers/ExrImgMediaReader.cpp
+++ b/Engine/Plugins/Media/ImgMedia/Source/ImgMedia/Private/Readers/ExrImgMediaReader.cpp
@@ -273,6 +273,11 @@ bool FExrImgMediaReader::ReadFrame(int32 FrameId, const TMap<int32, FImgMediaTil
 					if (ReadResult != Fail)
 					{
 						OutFrame->MipTilesPresent.Emplace(CurrentMipLevel, CurrentTileSelection);
+						
+						for (const FIntRect& Region : TileRegions)
+						{
+							OutFrame->NumTilesRead += Region.Area();
+						}
 					}
 					else
 					{
@@ -296,6 +301,7 @@ bool FExrImgMediaReader::ReadFrame(int32 FrameId, const TMap<int32, FImgMediaTil
 						InputFile.SetFrameBuffer(MipDataPtr, Dim);
 						InputFile.ReadPixels(0, Dim.Y - 1);
 						OutFrame->MipTilesPresent.Emplace(CurrentMipLevel, CurrentTileSelection);
+						OutFrame->NumTilesRead++;
 					}
 					else
 					{
@@ -415,6 +421,12 @@ bool FExrImgMediaReader::GetInfo(const FString& FilePath, FImgMediaFrameInfo& Ou
 		if (HeaderReader.ContainsMips())
 		{
 			OutInfo.NumMipLevels = HeaderReader.CalculateNumMipLevels(OutInfo.NumTiles);
+
+			SIZE_T SizeMip0 = OutInfo.UncompressedSize;
+			for (int32 Level = 1; Level < OutInfo.NumMipLevels; Level++)
+			{
+				OutInfo.UncompressedSize += SizeMip0 >> (2 * Level);
+			}
 		}
 	}
 	else
diff --git a/Engine/Plugins/Media/ImgMedia/Source/ImgMedia/Private/Readers/ExrImgMediaReaderGpu.cpp b/Engine/Plugins/Media/ImgMedia/Source/ImgMedia/Private/Readers/ExrImgMediaReaderGpu.cpp
--- a/Engine/Plugins/Media/ImgMedia/Source/ImgMedia/Private/Readers/ExrImgMediaReaderGpu.cpp
+++ b/Engine/Plugins/Media/ImgMedia/Source/ImgMedia/Private/Readers/ExrImgMediaReaderGpu.cpp
@@ -248,6 +248,11 @@ bool FExrImgMediaReaderGpu::ReadFrame(int32 FrameId, const TMap<int32, FImgMedia
 					}
 
 					ReadResult = ReadTiles(MipDataPtr, BufferSize, ImagePath, FrameId, TileRegionsToRead, ConverterParams, CurrentMipLevel, BufferRegionsToCopy);
+
+					for (const FIntRect& Region : TileRegionsToRead)
+					{
+						OutFrame->NumTilesRead += Region.Area();
+					}
 				}
 				else
 				{
@@ -259,7 +264,9 @@ bool FExrImgMediaReaderGpu::ReadFrame(int32 FrameId, const TMap<int32, FImgMedia
 					Viewports.Add(MoveTemp(Viewport));
 
 					ReadResult = ReadInChunks(MipDataPtr, ImagePath, FrameId, CurrentMipDim, BufferSize);
+					OutFrame->NumTilesRead++;
 				}
+
 				if (ReadResult == Success && CVarExrReaderUseUploadHeap.GetValueOnAnyThread())
 				{
 					TArray<int32> InMipLevels;
diff --git a/Engine/Plugins/Media/ImgMedia/Source/ImgMedia/Private/Readers/GenericImgMediaReader.cpp b/Engine/Plugins/Media/ImgMedia/Source/ImgMedia/Private/Readers/GenericImgMediaReader.cpp
--- a/Engine/Plugins/Media/ImgMedia/Source/ImgMedia/Private/Readers/GenericImgMediaReader.cpp
+++ b/Engine/Plugins/Media/ImgMedia/Source/ImgMedia/Private/Readers/GenericImgMediaReader.cpp
@@ -162,6 +162,7 @@ bool FGenericImgMediaReader::ReadFrame(int32 FrameId, const TMap<int32, FImgMedi
 			// Copy data to our buffer
 			FMemory::Memcpy(Buffer, RawData.GetData(), RawNum);
 			OutFrame->MipTilesPresent.Emplace(CurrentMipLevel, CurrentSelection);
+			OutFrame->NumTilesRead++;
 		}
 
 		// Next level.
diff --git a/Engine/Plugins/Media/ImgMedia/Source/ImgMedia/Private/Readers/IImgMediaReader.h b/Engine/Plugins/Media/ImgMedia/Source/ImgMedia/Private/Readers/IImgMediaReader.h
--- a/Engine/Plugins/Media/ImgMedia/Source/ImgMedia/Private/Readers/IImgMediaReader.h
+++ b/Engine/Plugins/Media/ImgMedia/Source/ImgMedia/Private/Readers/IImgMediaReader.h
@@ -71,12 +71,12 @@ struct FImgMediaFrame
 	/** Tiles present per mip level. */
 	TMap<int32, FImgMediaTileSelection> MipTilesPresent;
 
+	/** Total number of tiles read. */
+	int32 NumTilesRead = 0;
+
 	/** The frame's horizontal stride (in bytes). */
 	uint32 Stride = 0;
 
-	/** Uncompressed EXR files are read faster via plain read and GPU swizzling. This value is used by ExrImgMediaReaderGpu.*/
-	FTexture2DRHIRef Texture;
-
 	/** Sample converter is used by Media Texture Resource to convert the texture or data. */
 	TSharedPtr<IMediaTextureSampleConverter, ESPMode::ThreadSafe> SampleConverter;
 
diff --git a/Engine/Plugins/Media/ImgMedia/Source/ImgMediaEditor/Private/Widgets/SImgMediaBandwidth.cpp b/Engine/Plugins/Media/ImgMedia/Source/ImgMediaEditor/Private/Widgets/SImgMediaBandwidth.cpp
--- a/Engine/Plugins/Media/ImgMedia/Source/ImgMediaEditor/Private/Widgets/SImgMediaBandwidth.cpp
+++ b/Engine/Plugins/Media/ImgMedia/Source/ImgMediaEditor/Private/Widgets/SImgMediaBandwidth.cpp
@@ -17,6 +17,8 @@
 
 #define LOCTEXT_NAMESPACE "ImgMediaBandwidth"
 
+constexpr double SImgMediaBandwidth::BANDWIDTH_REFRESH_RATE;
+
 BEGIN_SLATE_FUNCTION_BUILD_OPTIMIZATION
 void SImgMediaBandwidth::Construct(const FArguments& InArgs)
 {
@@ -60,35 +62,46 @@ void SImgMediaBandwidth::Tick(const FGeometry& AllottedGeometry, const double In
 	// Call parent.
 	SCompoundWidget::Tick(AllottedGeometry, InCurrentTime, InDeltaTime);
 
-	// Loop over all players.
-	for (const FPlayerInfo& PlayerInfo : PlayerInfos)
+	if (InCurrentTime - LastBandwidthUpdateTime >= BANDWIDTH_REFRESH_RATE)
 	{
-		TSharedPtr< FImgMediaPlayer> Player = PlayerInfo.Player.Pin();
-		if (Player != nullptr)
+		// Loop over all players.
+		for (const FPlayerInfo& PlayerInfo : PlayerInfos)
 		{
-			// Update widgets with info.
-			// Set URL.
-			PlayerInfo.UrlTextBlock->SetText(FText::Format(
-				LOCTEXT("Url", "URL: {0}"),
-				FText::FromString(Player->GetUrl())));
-
-			// Get the loader.
-			TSharedPtr<FImgMediaLoader, ESPMode::ThreadSafe> Loader = Player->GetLoader();
-			if (Loader.IsValid())
+			TSharedPtr< FImgMediaPlayer> Player = PlayerInfo.Player.Pin();
+			if (Player != nullptr)
 			{
-				// Set cuurrent bandwidth.
-				float Bandwidth = Loader->GetCurrentBandwidth() / (1024.0f * 1024.0f);
-				PlayerInfo.CurrentBandwidthTextBlock->SetText(FText::Format(
-					LOCTEXT("Current", "Current: {0} MB/s"),
-					FText::AsNumber(Bandwidth)));
-
-				// Set minimum bandwidth.
-				Bandwidth = Loader->GetMinBandwidthForMipLevel0() / (1024.0f * 1024.0f);
-				PlayerInfo.MinimumBandwidthTextBlock->SetText(FText::Format(
-					LOCTEXT("Minimum", "Minimum: {0} MB/s"),
-					FText::AsNumber(Bandwidth)));
+				// Update widgets with info.
+				// Set URL.
+				PlayerInfo.UrlTextBlock->SetText(FText::Format(
+					LOCTEXT("Url", "URL: {0}"),
+					FText::FromString(Player->GetUrl())));
+
+				// Get the loader.
+				TSharedPtr<FImgMediaLoader, ESPMode::ThreadSafe> Loader = Player->GetLoader();
+				if (Loader.IsValid())
+				{
+					// Set current bandwidth.
+					float Bandwidth = Loader->GetCurrentBandwidth() / (1000.0f * 1000.0f);
+					PlayerInfo.CurrentBandwidthTextBlock->SetText(FText::Format(
+						LOCTEXT("Current", "Current: {0} MB/s"),
+						FText::AsNumber(Bandwidth)));
+
+					// Set effective bandwidth.
+					Bandwidth = Loader->GetEffectiveBandwidth() / (1000.0f * 1000.0f);
+					PlayerInfo.EffectiveBandwidthTextBlock->SetText(FText::Format(
+						LOCTEXT("Effective", "Effective: {0} MB/s"),
+						FText::AsNumber(Bandwidth)));
+
+					// Set required bandwidth.
+					Bandwidth = Loader->GetRequiredBandwidth() / (1000.0f * 1000.0f);
+					PlayerInfo.RequiredBandwidthTextBlock->SetText(FText::Format(
+						LOCTEXT("Required", "Required: {0} MB/s"),
+						FText::AsNumber(Bandwidth)));
+				}
 			}
 		}
+
+		LastBandwidthUpdateTime = InCurrentTime;
 	}
 }
 
@@ -141,7 +154,13 @@ void SImgMediaBandwidth::RefreshPlayersContainer()
 								+ SHorizontalBox::Slot()
 									.FillWidth(1.0f)
 									[
-										SAssignNew(PlayerInfo.MinimumBandwidthTextBlock, STextBlock)
+										SAssignNew(PlayerInfo.EffectiveBandwidthTextBlock, STextBlock)
+									]
+
+								+ SHorizontalBox::Slot()
+									.FillWidth(1.0f)
+									[
+										SAssignNew(PlayerInfo.RequiredBandwidthTextBlock, STextBlock)
 									]
 							]
 
diff --git a/Engine/Plugins/Media/ImgMedia/Source/ImgMediaEditor/Private/Widgets/SImgMediaBandwidth.h b/Engine/Plugins/Media/ImgMedia/Source/ImgMediaEditor/Private/Widgets/SImgMediaBandwidth.h
--- a/Engine/Plugins/Media/ImgMedia/Source/ImgMediaEditor/Private/Widgets/SImgMediaBandwidth.h
+++ b/Engine/Plugins/Media/ImgMedia/Source/ImgMediaEditor/Private/Widgets/SImgMediaBandwidth.h
@@ -37,12 +37,16 @@ private:
 		TWeakPtr<FImgMediaPlayer> Player;
 		TSharedPtr<STextBlock> UrlTextBlock;
 		TSharedPtr<STextBlock> CurrentBandwidthTextBlock;
-		TSharedPtr<STextBlock> MinimumBandwidthTextBlock;
+		TSharedPtr<STextBlock> EffectiveBandwidthTextBlock;
+		TSharedPtr<STextBlock> RequiredBandwidthTextBlock;
 	};
 
 	/** Array of all players. */
 	TArray<FPlayerInfo> PlayerInfos;
 	/** Widget to hold all widgets for all the players. */
 	TSharedPtr<SVerticalBox> PlayersContainer;
-
+	/** Timestamp to slow down bandwidth updates. */
+	double LastBandwidthUpdateTime = 0.0;
+	/** Bandwidth refresh rate in seconds. */
+	static constexpr double BANDWIDTH_REFRESH_RATE = 0.25;
 };
