ImgMedia: Partial tile calculation fix for mipped sequences without tiling, where the tile size remains zero.

#rnx

diff --git a/Engine/Plugins/Media/ImgMedia/Source/ImgMedia/Private/Assets/ImgMediaMipMapInfo.cpp b/Engine/Plugins/Media/ImgMedia/Source/ImgMedia/Private/Assets/ImgMediaMipMapInfo.cpp
--- a/Engine/Plugins/Media/ImgMedia/Source/ImgMedia/Private/Assets/ImgMediaMipMapInfo.cpp
+++ b/Engine/Plugins/Media/ImgMedia/Source/ImgMedia/Private/Assets/ImgMediaMipMapInfo.cpp
@@ -341,7 +341,6 @@ namespace {
 			}
 
 			const FIntPoint& SequenceTileNum = InSequenceInfo.TilingDescription.TileNum;
-			const FVector2f SequencePartialTileNum = FVector2f(InSequenceInfo.Dim) / FVector2f(InSequenceInfo.TilingDescription.TileSize);
 
 			// To avoid calculating tile corner mip levels multiple times over, we cache them in this array.
 			CornerMipLevelsCached.SetNum((SequenceTileNum.X + 1) * (SequenceTileNum.Y + 1));
@@ -397,7 +396,7 @@ namespace {
 					CurrentNumTiles.X = FMath::Max(1, FMath::CeilToInt(float(SequenceTileNum.X) / MipLevelDiv));
 					CurrentNumTiles.Y = FMath::Max(1, FMath::CeilToInt(float(SequenceTileNum.Y) / MipLevelDiv));
 
-					FVector2f CurrentPartialTileNum = SequencePartialTileNum / (float)MipLevelDiv;
+					FVector2f CurrentPartialTileNum = InSequenceInfo.GetPartialTileNum(CurrentMipLevel);
 
 					// Exclude subdivided tiles (enqueued below) that are not present (i.e. mipped sequences with odd number of tiles)
 					if (Tile.X >= CurrentNumTiles.X || Tile.Y >= CurrentNumTiles.Y)
@@ -567,7 +566,7 @@ namespace {
 
 			// Include all tiles containted in the visible UV region
 			const FIntPoint& SequenceTileNum = InSequenceInfo.TilingDescription.TileNum;
-			const FVector2f SequencePartialTileNum = FVector2f(InSequenceInfo.Dim) / FVector2f(InSequenceInfo.TilingDescription.TileSize);
+			const FVector2f SequencePartialTileNum = InSequenceInfo.GetPartialTileNum();
 
 			float PixelDimX = 1.0f / InSequenceInfo.Dim.X;
 			float PixelDimY = 1.0f / InSequenceInfo.Dim.Y;
diff --git a/Engine/Plugins/Media/ImgMedia/Source/ImgMedia/Private/Assets/ImgMediaMipMapInfo.h b/Engine/Plugins/Media/ImgMedia/Source/ImgMedia/Private/Assets/ImgMediaMipMapInfo.h
--- a/Engine/Plugins/Media/ImgMedia/Source/ImgMedia/Private/Assets/ImgMediaMipMapInfo.h
+++ b/Engine/Plugins/Media/ImgMedia/Source/ImgMedia/Private/Assets/ImgMediaMipMapInfo.h
@@ -38,6 +38,17 @@ struct FSequenceInfo
 	{
 		return (TilingDescription.TileNum.X > 1) || (TilingDescription.TileNum.Y > 1);
 	}
+	/** Get the fractional number of tiles, optionally at a specified mip level.*/
+	FORCEINLINE FVector2f GetPartialTileNum(const int32 MipLevel = 0) const
+	{
+		if (IsTiled())
+		{
+			ensure(MipLevel >= 0);
+			return (FVector2f(Dim) / FVector2f(TilingDescription.TileSize)) / float(1 << MipLevel);
+		}
+
+		return FVector2f::One();
+	}
 };
 
 /**
