Image Media: Removed Custom Exr option from Process Images.

diff --git a/Engine/Plugins/Media/ImgMedia/Source/ImgMediaEditor/Private/Widgets/ImgMediaProcessImagesOptions.h b/Engine/Plugins/Media/ImgMedia/Source/ImgMediaEditor/Private/Widgets/ImgMediaProcessImagesOptions.h
--- a/Engine/Plugins/Media/ImgMedia/Source/ImgMediaEditor/Private/Widgets/ImgMediaProcessImagesOptions.h
+++ b/Engine/Plugins/Media/ImgMedia/Source/ImgMediaEditor/Private/Widgets/ImgMediaProcessImagesOptions.h
@@ -23,10 +23,6 @@ public:
 	UPROPERTY(EditAnywhere, Category = Sequence)
 	FDirectoryPath OutputPath;
 
-	/** If checked, then use our custom format with a single file for all tiles. */
-	UPROPERTY(EditAnywhere, Category = Sequence)
-	bool bUseCustomFormat;
-
 	/** If checked, then enable mip mapping. */
 	UPROPERTY(EditAnywhere, Category = Sequence)
 	bool bEnableMipMapping;
diff --git a/Engine/Plugins/Media/ImgMedia/Source/ImgMediaEditor/Private/Widgets/SImgMediaProcessImages.cpp b/Engine/Plugins/Media/ImgMedia/Source/ImgMediaEditor/Private/Widgets/SImgMediaProcessImages.cpp
--- a/Engine/Plugins/Media/ImgMedia/Source/ImgMediaEditor/Private/Widgets/SImgMediaProcessImages.cpp
+++ b/Engine/Plugins/Media/ImgMedia/Source/ImgMediaEditor/Private/Widgets/SImgMediaProcessImages.cpp
@@ -197,7 +197,6 @@ void SImgMediaProcessImages::ProcessAllImages()
 {
 	TRACE_CPUPROFILER_EVENT_SCOPE(SImgMediaProcessImages::ProcessAllImages);
 
-	bool bUseCustomFormat = Options->bUseCustomFormat;
 	int32 InTileWidth = Options->bEnableTiling ? Options->TileSizeX : 0;
 	int32 InTileHeight = Options->bEnableTiling ? Options->TileSizeY : 0;
 	int32 TileBorder = 0; // Note: virtual texture support is shelved for now.
@@ -272,7 +271,7 @@ void SImgMediaProcessImages::ProcessAllImages()
 
 				TSharedPtr<IImageWrapper> ImageWrapper = ImageWrapperModule.CreateImageWrapper(ImageFormat);
 				Async(EAsyncExecution::Thread, [this, SequencePath, FileName, OutPath,
-					bUseCustomFormat, Ext, &NumActive,
+					Ext, &NumActive,
 					ImageWrapper, InTileWidth, InTileHeight,
 					TileBorder, bEnableMips, bHasAlphaChannel]() mutable
 				{
@@ -296,7 +295,7 @@ void SImgMediaProcessImages::ProcessAllImages()
 					// Import this image.
 					FString Name = FPaths::Combine(OutPath, FileName);
 					ProcessImageCustom(ImageWrapper, InTileWidth, InTileHeight, TileBorder,
-						bEnableMips, bHasAlphaChannel, Name, bUseCustomFormat);
+						bEnableMips, bHasAlphaChannel, Name, false /*bUseCustomFormat*/);
 
 					NumActive--;
 				});
@@ -923,7 +922,6 @@ void SImgMediaProcessImages::HandleProcessing()
 						int32 Width = RenderTarget->GetSurfaceWidth();
 						int32 Height = RenderTarget->GetSurfaceHeight();
 						int32 BitDepth = 16;
-						bool bUseCustomFormat = Options->bUseCustomFormat;
 						int32 InTileWidth = Options->bEnableTiling ? Options->TileSizeX : 0;
 						int32 InTileHeight = Options->bEnableTiling ? Options->TileSizeY : 0;
 						int32 TileBorder = 0; // Note: virtual texture support is shelved for now.
@@ -935,11 +933,11 @@ void SImgMediaProcessImages::HandleProcessing()
 
 						Async(EAsyncExecution::Thread, [this, RawData = MoveTemp(RawData), Width, Height, BitDepth,
 							InTileWidth, InTileHeight, TileBorder, bEnableMips,
-							bHasAlphaChannel, Name, bUseCustomFormat]() mutable
+							bHasAlphaChannel, Name]() mutable
 						{
 							ProcessImageCustomRawData(RawData, Width, Height, BitDepth,
 								InTileWidth, InTileHeight, TileBorder, bEnableMips,
-								bHasAlphaChannel, Name, bUseCustomFormat);
+								bHasAlphaChannel, Name, false /*bUseCustomFormat*/);
 						});
 					}
 					else
