Image Media Player: ExrImgReader GPU: Fixed an issue with cancelled frames still being rendered into texture.

diff --git a/Engine/Plugins/Media/ImgMedia/Source/ImgMedia/Private/Readers/ExrImgMediaReaderGpu.cpp b/Engine/Plugins/Media/ImgMedia/Source/ImgMedia/Private/Readers/ExrImgMediaReaderGpu.cpp
--- a/Engine/Plugins/Media/ImgMedia/Source/ImgMedia/Private/Readers/ExrImgMediaReaderGpu.cpp
+++ b/Engine/Plugins/Media/ImgMedia/Source/ImgMedia/Private/Readers/ExrImgMediaReaderGpu.cpp
@@ -288,17 +288,19 @@ bool FExrImgMediaReaderGpu::ReadFrame(int32 FrameId, const TMap<int32, FImgMedia
 						}
 					});
 				}
-				if (ReadResult != Fail)
-				{
-					OutFrame->MipTilesPresent.Emplace(CurrentMipLevel, CurrentTileSelection);
-				}
 			}
 			else
 			{
 				UE_LOG(LogImgMedia, Error, TEXT("Could not load %s"), *ImagePath);
 				return false;
 			}
-			if (ReadResult == Fail)
+
+			/* Error handling. */
+			if (ReadResult == Success)
+			{
+				OutFrame->MipTilesPresent.Emplace(CurrentMipLevel, CurrentTileSelection);
+			}
+			else if (ReadResult == Fail)
 			{
 				// Check if we have a compressed file.
 				FImgMediaFrameInfo Info;
@@ -319,6 +321,10 @@ bool FExrImgMediaReaderGpu::ReadFrame(int32 FrameId, const TMap<int32, FImgMedia
 
 				return FExrImgMediaReader::ReadFrame(FrameId, InMipTiles, OutFrame);
 			}
+			else if (ReadResult == Cancelled)
+			{
+				return false;
+			}
 		}
 	}
 
