ImgMedia: Fix material mip bias calculation such that the resolution fraction matches the renderer one again.

This may have changed with TSR updates in 20377247.

diff --git a/Engine/Plugins/Media/ImgMedia/Source/ImgMedia/Private/ImgMediaSceneViewExtension.cpp b/Engine/Plugins/Media/ImgMedia/Source/ImgMedia/Private/ImgMediaSceneViewExtension.cpp
--- a/Engine/Plugins/Media/ImgMedia/Source/ImgMedia/Private/ImgMediaSceneViewExtension.cpp
+++ b/Engine/Plugins/Media/ImgMedia/Source/ImgMedia/Private/ImgMediaSceneViewExtension.cpp
@@ -83,11 +83,18 @@ void FImgMediaSceneViewExtension::BeginRenderViewFamily(FSceneViewFamily& InView
 			Info.PrimitiveComponentIds = View->HiddenPrimitives;
 		}
 
-		// View->MaterialTextureMipBias is only set later in rendering: we replicate the logic here.
+		/* View->MaterialTextureMipBias is only set later in rendering so we replicate here the calculations
+		 * found in FSceneRenderer::PreVisibilityFrameSetup.*/
 		if (View->PrimaryScreenPercentageMethod == EPrimaryScreenPercentageMethod::TemporalUpscale)
 		{
-			Info.MaterialTextureMipBias = -(FMath::Max(-FMath::Log2(ResolutionFraction), 0.0f)) + CVarMinAutomaticViewMipBiasOffset->GetValueOnGameThread();
+			const float EffectivePrimaryResolutionFraction = float(Info.ViewportRect.Width()) / (View->UnscaledViewRect.Width() * InViewFamily.SecondaryViewFraction);
+			Info.MaterialTextureMipBias = -(FMath::Max(-FMath::Log2(EffectivePrimaryResolutionFraction), 0.0f)) + CVarMinAutomaticViewMipBiasOffset->GetValueOnGameThread();
 			Info.MaterialTextureMipBias = FMath::Max(Info.MaterialTextureMipBias, CVarMinAutomaticViewMipBias->GetValueOnGameThread());
+
+			if (!ensureMsgf(!FMath::IsNaN(Info.MaterialTextureMipBias) && FMath::IsFinite(Info.MaterialTextureMipBias), TEXT("Calculated material texture mip bias is invalid, defaulting to zero.")))
+			{
+				Info.MaterialTextureMipBias = 0.0f;
+			}
 		}
 		else
 		{
