ImgMedia: Avoid the very slow FRgbaInputFile constructor when the header reader already extracts the necessary info.

diff --git a/Engine/Plugins/Media/ImgMedia/Source/ImgMedia/Private/Readers/ExrImgMediaReader.cpp b/Engine/Plugins/Media/ImgMedia/Source/ImgMedia/Private/Readers/ExrImgMediaReader.cpp
--- a/Engine/Plugins/Media/ImgMedia/Source/ImgMedia/Private/Readers/ExrImgMediaReader.cpp
+++ b/Engine/Plugins/Media/ImgMedia/Source/ImgMedia/Private/Readers/ExrImgMediaReader.cpp
@@ -333,11 +333,9 @@ TSharedPtr<IImgMediaReader, ESPMode::ThreadSafe> FExrImgMediaReader::GetReader(c
 	FIntPoint TileSize(EForceInit::ForceInitToZero);
 
 #if defined(PLATFORM_WINDOWS) && PLATFORM_WINDOWS
-	FRgbaInputFile InputFile(FirstImageInSequencePath, 2);
-	if (InputFile.HasInputFile() == false)
+	if (!FPaths::FileExists(FirstImageInSequencePath))
 	{
-		TSharedPtr<IImgMediaReader, ESPMode::ThreadSafe> Ptr;
-		return Ptr;
+		return nullptr;
 	}
 	
 	FImgMediaFrameInfo Info;
@@ -346,15 +344,10 @@ TSharedPtr<IImgMediaReader, ESPMode::ThreadSafe> FExrImgMediaReader::GetReader(c
 		return MakeShareable(new FExrImgMediaReader(InLoader));
 	}
 
-	// Is this our custom format?
-	int32 CustomFormat = 0;
-	InputFile.GetIntAttribute(IImgMediaModule::CustomFormatAttributeName.Resolve().ToString(), CustomFormat);
-	bIsCustomFormat = CustomFormat > 0;
+	bIsCustomFormat = Info.FormatName.Equals(TEXT("EXR CUSTOM"));
 	if (bIsCustomFormat)
 	{
-		// Get tile size.
-		InputFile.GetIntAttribute(IImgMediaModule::CustomFormatTileWidthAttributeName.Resolve().ToString(), TileSize.X);
-		InputFile.GetIntAttribute(IImgMediaModule::CustomFormatTileHeightAttributeName.Resolve().ToString(), TileSize.Y);
+		TileSize = Info.TileDimensions;
 	}
 
 	// Check GetCompressionName of OpenExrWrapper for other compression names.
