ImgMedia: Added support for the mip map bias for both the loader and media texture sampler, as an advanced MediaPlateComponent parameter.

(Depends on 21700426.)

diff --git a/Engine/Plugins/Media/ImgMedia/Source/ImgMedia/Private/Assets/ImgMediaMipMapInfo.cpp b/Engine/Plugins/Media/ImgMedia/Source/ImgMedia/Private/Assets/ImgMediaMipMapInfo.cpp
--- a/Engine/Plugins/Media/ImgMedia/Source/ImgMedia/Private/Assets/ImgMediaMipMapInfo.cpp
+++ b/Engine/Plugins/Media/ImgMedia/Source/ImgMedia/Private/Assets/ImgMediaMipMapInfo.cpp
@@ -250,9 +250,9 @@ namespace {
 }
 
 
-FImgMediaMipMapObjectInfo::FImgMediaMipMapObjectInfo(UMeshComponent* InMeshComponent, float InLODBias)
+FImgMediaMipMapObjectInfo::FImgMediaMipMapObjectInfo(UMeshComponent* InMeshComponent, float InMipMapBias)
 	: MeshComponent(InMeshComponent)
-	, LODBias(InLODBias)
+	, MipMapBias(InMipMapBias)
 {
 
 }
@@ -308,7 +308,7 @@ namespace {
 		{
 			float DistX = FVector2D::DistSquared(TexelScreenSpace[0], TexelScreenSpace[1]);
 			float DistY = FVector2D::DistSquared(TexelScreenSpace[0], TexelScreenSpace[2]);
-			OutMipLevel = FMath::Max(0.5f * (float)FMath::Log2(1.0f / FMath::Max(DistX, DistY)), 0.0f); // ~ log2(sqrt(delta))
+			OutMipLevel = 0.5f * (float)FMath::Log2(1.0f / FMath::Max(DistX, DistY)); // ~ log2(sqrt(delta))
 		}
 
 		return bValid;
@@ -317,8 +317,8 @@ namespace {
 	class FPlaneObjectInfo : public FImgMediaMipMapObjectInfo
 	{
 	public:
-		FPlaneObjectInfo(UMeshComponent* InMeshComponent, float InLODBias = 0.0f)
-			: FImgMediaMipMapObjectInfo(InMeshComponent, InLODBias)
+		FPlaneObjectInfo(UMeshComponent* InMeshComponent, float InMipMapBias = 0.0f)
+			: FImgMediaMipMapObjectInfo(InMeshComponent, InMipMapBias)
 			, PlaneSize(FVector::ZeroVector)
 		{
 			// Get size of object.
@@ -444,7 +444,7 @@ namespace {
 
 									if (CalculateMipLevel(ViewInfo, CornersWS, CornersWS + TexelOffsetXWS, CornersWS + TexelOffsetYWS, CalculatedLevel))
 									{
-										CalculatedLevel += LODBias + ViewInfo.MaterialTextureMipBias;
+										CalculatedLevel += MipMapBias + ViewInfo.MaterialTextureMipBias;
 
 										SetCachedMipLevel(BaseLevelCorner.X, BaseLevelCorner.Y, MaxCornerX, CalculatedLevel);
 										bValidLevel = true;
@@ -546,8 +546,8 @@ namespace {
 	class FSphereObjectInfo : public FImgMediaMipMapObjectInfo
 	{
 	public:
-		FSphereObjectInfo(UMeshComponent* InMeshComponent, float InLODBias, FVector2D InMeshRange)
-			: FImgMediaMipMapObjectInfo(InMeshComponent, InLODBias)
+		FSphereObjectInfo(UMeshComponent* InMeshComponent, float InMipMapBias, FVector2D InMeshRange)
+			: FImgMediaMipMapObjectInfo(InMeshComponent, InMipMapBias)
 			, MeshRange(InMeshRange)
 		{
 		}
@@ -613,7 +613,7 @@ namespace {
 
 							if (CalculateMipLevel(ViewInfo, TileCenterWS, TexelOffXWS, TexelOffYWS, CalculatedLevel))
 							{
-								CalculatedLevel += LODBias + ViewInfo.MaterialTextureMipBias;
+								CalculatedLevel += MipMapBias + ViewInfo.MaterialTextureMipBias;
 
 								MipLevelRange[0] = FMath::Clamp((int32)CalculatedLevel, 0, MaxLevel);
 								MipLevelRange[1] = FMath::CeilToInt32(CalculatedLevel);
@@ -670,7 +670,7 @@ FImgMediaMipMapInfo::~FImgMediaMipMapInfo()
 	ClearAllObjects();
 }
 
-void FImgMediaMipMapInfo::AddObject(AActor* InActor, float LODBias, EMediaTextureVisibleMipsTiles MeshType,
+void FImgMediaMipMapInfo::AddObject(AActor* InActor, float InMipMapBias, EMediaTextureVisibleMipsTiles InMeshMode,
 	FVector2D MeshRange)
 {
 	if (InActor != nullptr)
@@ -680,16 +680,16 @@ void FImgMediaMipMapInfo::AddObject(AActor* InActor, float LODBias, EMediaTextur
 		UMeshComponent* MeshComponent = Cast<UMeshComponent>(InActor->FindComponentByClass(UMeshComponent::StaticClass()));
 		if (MeshComponent != nullptr)
 		{
-			switch (MeshType)
+			switch (InMeshMode)
 			{
 			case EMediaTextureVisibleMipsTiles::Plane:
-				Objects.Add(new FPlaneObjectInfo(MeshComponent, LODBias));
+				Objects.Add(new FPlaneObjectInfo(MeshComponent, InMipMapBias));
 				break;
 			case EMediaTextureVisibleMipsTiles::Sphere:
-				Objects.Add(new FSphereObjectInfo(MeshComponent, LODBias, MeshRange));
+				Objects.Add(new FSphereObjectInfo(MeshComponent, InMipMapBias, MeshRange));
 				break;
 			default:
-				Objects.Add(new FImgMediaMipMapObjectInfo(MeshComponent, LODBias));
+				Objects.Add(new FImgMediaMipMapObjectInfo(MeshComponent, InMipMapBias));
 				break;
 			}
 		}
diff --git a/Engine/Plugins/Media/ImgMedia/Source/ImgMedia/Private/Assets/ImgMediaMipMapInfo.h b/Engine/Plugins/Media/ImgMedia/Source/ImgMedia/Private/Assets/ImgMediaMipMapInfo.h
--- a/Engine/Plugins/Media/ImgMedia/Source/ImgMedia/Private/Assets/ImgMediaMipMapInfo.h
+++ b/Engine/Plugins/Media/ImgMedia/Source/ImgMedia/Private/Assets/ImgMediaMipMapInfo.h
@@ -185,7 +185,7 @@ private:
 class FImgMediaMipMapObjectInfo
 {
 public:
-	FImgMediaMipMapObjectInfo(UMeshComponent* InMeshComponent, float InLODBias = 0.0f);
+	FImgMediaMipMapObjectInfo(UMeshComponent* InMeshComponent, float InMipMapBias = 0.0f);
 	virtual ~FImgMediaMipMapObjectInfo() = default;
 
 	/**
@@ -206,8 +206,8 @@ public:
 protected:
 	/** Mesh component onto which the media is displayed. */
 	TWeakObjectPtr<class UMeshComponent> MeshComponent;
-	/** LOD bias for the mipmap level. */
-	float LODBias;
+	/** Mipmap level bias. */
+	float MipMapBias;
 };
 
 /**
@@ -223,12 +223,12 @@ public:
 	 * This object is using our img sequence.
 	 *
 	 * @param InActor Object using our img sequence.
-	 * @param LODBias Bias used for mip-level calculations (Note: must be mirrored in a material manually).
-	 * @param MeshType Plane, sphere or none mesh type for mip-tile calculations.
+	 * @param InMipMapBias Bias used for mip-level calculations.
+	 * @param InMeshMode Plane, sphere or none mesh type for mip-tile calculations.
 	 * @param MeshRange Arc size in degrees used for visible tiles calculations, specific to the sphere.
 	 */
-	void AddObject(AActor* InActor, float LODBias = 0.0f,
-		EMediaTextureVisibleMipsTiles MeshType = EMediaTextureVisibleMipsTiles::None,
+	void AddObject(AActor* InActor, float InMipMapBias = 0.0f,
+		EMediaTextureVisibleMipsTiles InMeshMode = EMediaTextureVisibleMipsTiles::None,
 		FVector2D MeshRange = FVector2D(360.0f, 180.0f));
 
 	/**
diff --git a/Engine/Plugins/Media/MediaPlate/Source/MediaPlate/Private/MediaPlate.cpp b/Engine/Plugins/Media/MediaPlate/Source/MediaPlate/Private/MediaPlate.cpp
--- a/Engine/Plugins/Media/MediaPlate/Source/MediaPlate/Private/MediaPlate.cpp
+++ b/Engine/Plugins/Media/MediaPlate/Source/MediaPlate/Private/MediaPlate.cpp
@@ -113,13 +113,11 @@ void AMediaPlate::UseDefaultMaterial()
 
 void AMediaPlate::ApplyCurrentMaterial()
 {
-	if (StaticMeshComponent != nullptr)
+	UMaterialInterface* MaterialInterface = GetCurrentMaterial();
+	
+	if ((MaterialInterface != nullptr) && (LastMaterial != MaterialInterface))
 	{
-		UMaterialInterface* MaterialInterface = StaticMeshComponent->GetMaterial(0);
-		if ((MaterialInterface != nullptr) && (LastMaterial != MaterialInterface))
-		{
-			ApplyMaterial(MaterialInterface);
-		}
+		ApplyMaterial(MaterialInterface);
 	}
 }
 
@@ -186,6 +184,16 @@ void AMediaPlate::ApplyMaterial(UMaterialInterface* Material)
 	}
 }
 
+UMaterialInterface* AMediaPlate::GetCurrentMaterial() const
+{
+	if (StaticMeshComponent != nullptr)
+	{
+		return StaticMeshComponent->GetMaterial(0);
+	}
+
+	return nullptr;
+}
+
 void AMediaPlate::OnPreSaveWorld(UWorld* InWorld, FObjectPreSaveContext ObjectSaveContext)
 {
 	// We need to remove our asset user data before saving, as we do not need to save it out
diff --git a/Engine/Plugins/Media/MediaPlate/Source/MediaPlate/Private/MediaPlateComponent.cpp b/Engine/Plugins/Media/MediaPlate/Source/MediaPlate/Private/MediaPlateComponent.cpp
--- a/Engine/Plugins/Media/MediaPlate/Source/MediaPlate/Private/MediaPlateComponent.cpp
+++ b/Engine/Plugins/Media/MediaPlate/Source/MediaPlate/Private/MediaPlateComponent.cpp
@@ -10,6 +10,7 @@
 #include "Materials/Material.h"
 #include "MediaComponent.h"
 #include "MediaPlateModule.h"
+#include "MediaPlate.h"
 #include "MediaPlayer.h"
 #include "MediaPlaylist.h"
 #include "MediaSource.h"
@@ -97,6 +98,7 @@ void UMediaPlateComponent::OnRegister()
 	// Set up media texture.
 	if (MediaTexture != nullptr)
 	{
+		MediaTexture->SetMipMapBias(MipMapBias);
 		MediaTexture->SetMediaPlayer(MediaPlayer);
 		MediaTexture->UpdateResource();
 	}
@@ -296,7 +298,7 @@ void UMediaPlateComponent::RegisterWithMediaTextureTracker()
 	// Set up object.
 	MediaTextureTrackerObject = MakeShared<FMediaTextureTrackerObject, ESPMode::ThreadSafe>();
 	MediaTextureTrackerObject->Object = GetOwner();
-	MediaTextureTrackerObject->MipMapLODBias = 0.0f;
+	MediaTextureTrackerObject->MipMapLODBias = MipMapBias;
 	MediaTextureTrackerObject->VisibleMipsTilesCalculations = VisibleMipsTilesCalculations;
 	MediaTextureTrackerObject->MeshRange = MeshRange;
 
@@ -439,6 +441,18 @@ void UMediaPlateComponent::TickOutput()
 	}
 }
 
+void UMediaPlateComponent::RestartPlayer()
+{
+	if (MediaPlayer != nullptr)
+	{
+		if (MediaPlayer->IsPlaying())
+		{
+			MediaPlayer->Close();
+			Open();
+		}
+	}
+}
+
 void UMediaPlateComponent::StopClockSink()
 {
 	if (ClockSink.IsValid())
@@ -617,21 +631,15 @@ void UMediaPlateComponent::OnMediaEnd()
 
 #if WITH_EDITOR
 
-void UMediaPlateComponent::OnVisibleMipsTilesCalculationsChange()
+void UMediaPlateComponent::SetVisibleMipsTilesCalculations(EMediaTextureVisibleMipsTiles InVisibleMipsTilesCalculations)
 {
+	VisibleMipsTilesCalculations = InVisibleMipsTilesCalculations;
+
 	if (MediaTextureTrackerObject != nullptr)
 	{
 		MediaTextureTrackerObject->VisibleMipsTilesCalculations = VisibleMipsTilesCalculations;
 
-		// Propagate the change by restarting the player if it is currently playing.
-		if (MediaPlayer != nullptr)
-		{
-			if (MediaPlayer->IsPlaying())
-			{
-				MediaPlayer->Close();
-				Open();
-			}
-		}
+		RestartPlayer();
 	}
 }
 
@@ -674,7 +682,32 @@ void UMediaPlateComponent::PostEditChangeProperty(FPropertyChangedEvent& Propert
 	}
 	else if (PropertyChangedEvent.GetPropertyName() == GET_MEMBER_NAME_CHECKED(ThisClass, VisibleMipsTilesCalculations))
 	{
-		OnVisibleMipsTilesCalculationsChange();
+		if (MediaTextureTrackerObject != nullptr)
+		{
+			MediaTextureTrackerObject->VisibleMipsTilesCalculations = VisibleMipsTilesCalculations;
+			
+			RestartPlayer();
+		}
+	}
+	else if (PropertyChangedEvent.GetPropertyName() == GET_MEMBER_NAME_CHECKED(ThisClass, MipMapBias))
+	{
+		if (MediaTextureTrackerObject != nullptr)
+		{
+			MediaTextureTrackerObject->MipMapLODBias = MipMapBias;
+
+			// Note: Media texture bias automatically updated by UMediaPlateComponent::OnRegister().
+
+			// Update material sampler with new bias
+			if (AMediaPlate* MediaPlate = GetOwner<AMediaPlate>())
+			{
+				if (UMaterialInterface* Material = MediaPlate->GetCurrentMaterial())
+				{
+					Material->PostEditChange();
+				}
+			}
+
+			RestartPlayer();
+		}
 	}
 }
 
diff --git a/Engine/Plugins/Media/MediaPlate/Source/MediaPlate/Public/MediaPlate.h b/Engine/Plugins/Media/MediaPlate/Source/MediaPlate/Public/MediaPlate.h
--- a/Engine/Plugins/Media/MediaPlate/Source/MediaPlate/Public/MediaPlate.h
+++ b/Engine/Plugins/Media/MediaPlate/Source/MediaPlate/Public/MediaPlate.h
@@ -43,6 +43,11 @@ public:
 	 */
 	void ApplyCurrentMaterial();
 	void ApplyMaterial(UMaterialInterface* InMaterial);
+
+	/** Get the current static mesh material, at index 0. */
+	UMaterialInterface* GetCurrentMaterial() const;
+	
+	/** Get the last material assigned to the static mesh, at index 0. */
 	UMaterialInterface* GetLastMaterial() const { return LastMaterial; }
 #endif // WITH_EDITOR
 
diff --git a/Engine/Plugins/Media/MediaPlate/Source/MediaPlate/Public/MediaPlateComponent.h b/Engine/Plugins/Media/MediaPlate/Source/MediaPlate/Public/MediaPlateComponent.h
--- a/Engine/Plugins/Media/MediaPlate/Source/MediaPlate/Public/MediaPlateComponent.h
+++ b/Engine/Plugins/Media/MediaPlate/Source/MediaPlate/Public/MediaPlateComponent.h
@@ -150,13 +150,6 @@ public:
 	UPROPERTY(EditAnywhere, AdvancedDisplay, Category = "MediaPlate")
 	FMediaSourceCacheSettings CacheSettings;
 
-	/**
-	 * Specify type of mesh used for visible mips and tiles calculations.
-	 * (Using the provided plane and sphere meshes only.)
-	 */
-	UPROPERTY(EditAnywhere, AdvancedDisplay, Category = "MediaPlate")
-	EMediaTextureVisibleMipsTiles VisibleMipsTilesCalculations;
-
 	/** Set the arc size in degrees used for visible mips and tiles calculations, specific to the sphere. */
 	UFUNCTION(BlueprintCallable, Category = "Media|MediaPlateComponent")
 	void SetMeshRange(FVector2D InMeshRange);
@@ -200,17 +193,25 @@ public:
 	UFUNCTION(BlueprintCallable, Category = "Media|MediaPlateComponent")
 	void SetLetterboxAspectRatio(float AspectRatio);
 
+#if WITH_EDITOR
+	/**
+	 * Call this to get the mip tile calculations mesh mode.
+	 */
+	EMediaTextureVisibleMipsTiles GetVisibleMipsTilesCalculations() const { return VisibleMipsTilesCalculations; }
+	/**
+	 * Call this to set the mip tile calculations mesh mode. (Note: restarts playback to apply changes.)
+	 */
+	void SetVisibleMipsTilesCalculations(EMediaTextureVisibleMipsTiles InVisibleMipsTilesCalculations);
+#endif
+
 	/**
 	 * Called from the media clock.
 	 */
 	void TickOutput();
 
-#if WITH_EDITOR
-	/** Propagate visible calculation changes to the tracker/player objects, optionally restarting playback if needed. */
-	void OnVisibleMipsTilesCalculationsChange();
-#endif
-
 private:
+	void RestartPlayer();
+
 	/** If true then only allow playback when the media plate is visible. */
 	UPROPERTY(EditAnywhere, BlueprintReadOnly, Category = "MediaPlate", meta = (AllowPrivateAccess = true))
 	bool bPlayOnlyWhenVisible = false;
@@ -219,6 +220,14 @@ private:
 	UPROPERTY(EditAnywhere, Blueprintgetter = GetLoop, BlueprintSetter = SetLoop, Category = "MediaPlate", meta = (AllowPrivateAccess = true))
 	bool bLoop;
 
+	/** Visible mips and tiles calculation mode for the supported mesh types in MediaPlate. (Player restart on change.) */
+	UPROPERTY(EditAnywhere, AdvancedDisplay, Category = "MediaPlate", meta = (AllowPrivateAccess = true))
+	EMediaTextureVisibleMipsTiles VisibleMipsTilesCalculations;
+
+	/** Media texture mip map bias shared between the (image sequence) loader and the media texture sampler. (Player restart on change.) */
+	UPROPERTY(EditAnywhere, AdvancedDisplay, Category = "MediaPlate", meta = (AllowPrivateAccess = true, UIMin = "-16.0", UIMax = "15.99"))
+	float MipMapBias = 0.0f;
+
 	/** If > 0, then this is the aspect ratio of our screen and 
 	 * letterboxes will be added if the media is smaller than the screen. */
 	UPROPERTY()
diff --git a/Engine/Plugins/Media/MediaPlate/Source/MediaPlateEditor/Private/MediaPlateCustomization.cpp b/Engine/Plugins/Media/MediaPlate/Source/MediaPlateEditor/Private/MediaPlateCustomization.cpp
--- a/Engine/Plugins/Media/MediaPlate/Source/MediaPlateEditor/Private/MediaPlateCustomization.cpp
+++ b/Engine/Plugins/Media/MediaPlate/Source/MediaPlateEditor/Private/MediaPlateCustomization.cpp
@@ -71,7 +71,7 @@ void FMediaPlateCustomization::CustomizeDetails(IDetailLayoutBuilder& DetailBuil
 		if (MediaPlate.IsValid())
 		{
 			MediaPlatesList.Add(MediaPlate);
-			MeshMode = MediaPlate->VisibleMipsTilesCalculations;
+			MeshMode = MediaPlate->GetVisibleMipsTilesCalculations();
 		}
 	}
 
@@ -652,8 +652,7 @@ void FMediaPlateCustomization::SetMeshMode(EMediaTextureVisibleMipsTiles InMode)
 			if (MediaPlate != nullptr)
 			{
 				// Update the setting in the media plate.
-				MediaPlate->VisibleMipsTilesCalculations = MeshMode;
-				MediaPlate->OnVisibleMipsTilesCalculationsChange();
+				MediaPlate->SetVisibleMipsTilesCalculations(MeshMode);
 
 				// Set the appropriate mesh.
 				if (MeshMode == EMediaTextureVisibleMipsTiles::Plane)
