ImgMedia: UI adjustments in process EXR.

diff --git a/Engine/Plugins/Media/ImgMedia/Source/ImgMediaEditor/Private/ImgMediaEditorModule.cpp b/Engine/Plugins/Media/ImgMedia/Source/ImgMediaEditor/Private/ImgMediaEditorModule.cpp
--- a/Engine/Plugins/Media/ImgMedia/Source/ImgMediaEditor/Private/ImgMediaEditorModule.cpp
+++ b/Engine/Plugins/Media/ImgMedia/Source/ImgMediaEditor/Private/ImgMediaEditorModule.cpp
@@ -15,7 +15,7 @@
 #include "Widgets/Docking/SDockTab.h"
 #include "Widgets/SImgMediaBandwidth.h"
 #include "Widgets/SImgMediaCache.h"
-#include "Widgets/SImgMediaProcessImages.h"
+#include "Widgets/SImgMediaProcessEXR.h"
 #include "WorkspaceMenuStructure.h"
 #include "WorkspaceMenuStructureModule.h"
 
@@ -140,11 +140,11 @@ protected:
 			.SetIcon(FSlateIcon(FAppStyle::GetAppStyleSetName(), "SequenceRecorder.TabIcon"));
 
 		// Add process images tab.
-		FGlobalTabmanager::Get()->RegisterNomadTabSpawner(ImgMediaProcessImagesTabName,
-			FOnSpawnTab::CreateStatic(&FImgMediaEditorModule::SpawnProcessImagesTab))
+		FGlobalTabmanager::Get()->RegisterNomadTabSpawner(ImgMediaProcessEXRTabName,
+			FOnSpawnTab::CreateStatic(&FImgMediaEditorModule::SpawnProcessEXRTab))
 			.SetGroup(MediaBrowserGroup)
-			.SetDisplayName(LOCTEXT("ImgMediaProcessImagesTabTitle", "Process Images"))
-			.SetTooltipText(LOCTEXT("ImgMediaProcessImagesTooltipText", "Open the Process Images tab."))
+			.SetDisplayName(LOCTEXT("ImgMediaProcessImagesTabTitle", "Process EXR"))
+			.SetTooltipText(LOCTEXT("ImgMediaProcessImagesTooltipText", "Open the Process EXR tab."))
 			.SetIcon(FSlateIcon(FAppStyle::GetAppStyleSetName(), "LevelEditor.Tabs.Viewports"));
 	}
 
@@ -152,7 +152,7 @@ protected:
 	{
 		if (FSlateApplication::IsInitialized())
 		{
-			FGlobalTabmanager::Get()->UnregisterNomadTabSpawner(ImgMediaProcessImagesTabName);
+			FGlobalTabmanager::Get()->UnregisterNomadTabSpawner(ImgMediaProcessEXRTabName);
 			FGlobalTabmanager::Get()->UnregisterNomadTabSpawner(ImgMediaCacheTabName);
 			FGlobalTabmanager::Get()->UnregisterNomadTabSpawner(ImgMediaBandwidthTabName);
 		}
@@ -176,12 +176,12 @@ protected:
 			];
 	}
 
-	static TSharedRef<SDockTab> SpawnProcessImagesTab(const FSpawnTabArgs& SpawnTabArgs)
+	static TSharedRef<SDockTab> SpawnProcessEXRTab(const FSpawnTabArgs& SpawnTabArgs)
 	{
 		return SNew(SDockTab)
 			.TabRole(ETabRole::NomadTab)
 			[
-				SNew(SImgMediaProcessImages)
+				SNew(SImgMediaProcessEXR)
 			];
 	}
 
@@ -215,7 +215,7 @@ private:
 	FName CustomizedStructName;
 
 	/** Names for tabs. */
-	static FLazyName ImgMediaProcessImagesTabName;
+	static FLazyName ImgMediaProcessEXRTabName;
 
 	/** The collection of registered asset type actions. */
 	TArray<TSharedRef<IAssetTypeActions>> RegisteredAssetTypeActions;
@@ -224,7 +224,7 @@ private:
 	TArray<TWeakPtr<FImgMediaPlayer>> MediaPlayers;
 };
 
-FLazyName FImgMediaEditorModule::ImgMediaProcessImagesTabName(TEXT("ImgMediaProcessImages"));
+FLazyName FImgMediaEditorModule::ImgMediaProcessEXRTabName(TEXT("ImgMediaProcessEXR"));
 
 IMPLEMENT_MODULE(FImgMediaEditorModule, ImgMediaEditor);
 
diff --git a/Engine/Plugins/Media/ImgMedia/Source/ImgMediaEditor/Private/Widgets/ImgMediaProcessEXROptions.cpp b/Engine/Plugins/Media/ImgMedia/Source/ImgMediaEditor/Private/Widgets/ImgMediaProcessEXROptions.cpp
new file mode 100644
--- /dev/null
+++ b/Engine/Plugins/Media/ImgMedia/Source/ImgMediaEditor/Private/Widgets/ImgMediaProcessEXROptions.cpp
@@ -0,0 +1,15 @@
+// Copyright Epic Games, Inc. All Rights Reserved.
+
+#include "ImgMediaProcessEXROptions.h"
+
+UImgMediaProcessEXROptions::UImgMediaProcessEXROptions(const FObjectInitializer& ObjectInitializer)
+	: Super(ObjectInitializer)
+{
+	MipLevelTints.Add(FLinearColor::Red);
+	MipLevelTints.Add(FLinearColor::Green);
+	MipLevelTints.Add(FLinearColor::Blue);
+	MipLevelTints.Add(FLinearColor::Yellow);
+	MipLevelTints.Add(FLinearColor(0.0f, 1.0f, 1.0f));
+	MipLevelTints.Add(FLinearColor(1.0f, 0.0f, 1.0f));
+}
+
diff --git a/Engine/Plugins/Media/ImgMedia/Source/ImgMediaEditor/Private/Widgets/ImgMediaProcessEXROptions.h b/Engine/Plugins/Media/ImgMedia/Source/ImgMediaEditor/Private/Widgets/ImgMediaProcessEXROptions.h
new file mode 100644
--- /dev/null
+++ b/Engine/Plugins/Media/ImgMedia/Source/ImgMediaEditor/Private/Widgets/ImgMediaProcessEXROptions.h
@@ -0,0 +1,65 @@
+// Copyright Epic Games, Inc. All Rights Reserved.
+
+/**
+ * Contains options for processing image sequences.
+ */
+
+#pragma once
+
+#include "CoreMinimal.h"
+#include "UObject/ObjectMacros.h"
+#include "ImgMediaProcessEXROptions.generated.h"
+
+UCLASS(MinimalAPI)
+class UImgMediaProcessEXROptions : public UObject
+{
+	GENERATED_UCLASS_BODY()
+public:
+	/** The directory that contains the image sequence files. */
+	UPROPERTY(EditAnywhere, Category = Sequence)
+	FFilePath InputPath;
+
+	/** The directory to output the processed image sequence files to. */
+	UPROPERTY(EditAnywhere, Category = Sequence)
+	FDirectoryPath OutputPath;
+
+	/** If checked, then enable mip mapping. */
+	UPROPERTY(EditAnywhere, Category = Sequence)
+	bool bEnableMipMapping = true;
+
+	/** If checked, then enable tiling. */
+	UPROPERTY(EditAnywhere, Category = Tiles)
+	bool bEnableTiling = true;
+	
+	/** Width of a tile in pixels. If 0, then do not make tiles. */
+	UPROPERTY(EditAnywhere, Transient, Category = Tiles, meta = (EditCondition = "bEnableTiling", ClampMin = "0.0"))
+	int32 TileSizeX = 256;
+
+	/** Height of a tile in pixels. If 0, then do not make tiles. */
+	UPROPERTY(EditAnywhere, Transient, Category = Tiles, meta = (EditCondition = "bEnableTiling", ClampMin = "0.0"))
+	int32 TileSizeY = 256;
+
+	/** Number of threads to use when processing. If 0 then this number is set automatically. */
+	UPROPERTY(EditAnywhere, Transient, Category = Processing, meta = (ClampMin = "0.0"))
+	int32 NumThreads = 0;
+
+	/** Use a player to read in the image. */
+	UPROPERTY(EditAnywhere, Transient, Category = Processing)
+	bool bUsePlayer = false;
+
+	/** Number of tiles in the X direction. If 0, then there are no tiles. */
+	UPROPERTY(VisibleAnywhere, Transient, Category = Tiles)
+	int32 NumTilesX = 0;
+
+	/** Number of tiles in the Y direction. If 0, then there are no tiles. */
+	UPROPERTY(VisibleAnywhere, Transient, Category = Tiles)
+	int32 NumTilesY = 0;
+
+	/** Tint each mip level a different colour to help with debugging. */
+	UPROPERTY(EditAnywhere, Transient, Category = Debug)
+	bool bEnableMipLevelTint = false;
+
+	/** Colour to tint each mip level. */
+	UPROPERTY(EditAnywhere, Transient, Category = Debug)
+	TArray<FLinearColor> MipLevelTints;
+};
diff --git a/Engine/Plugins/Media/ImgMedia/Source/ImgMediaEditor/Private/Widgets/ImgMediaProcessImagesOptions.cpp b/Engine/Plugins/Media/ImgMedia/Source/ImgMediaEditor/Private/Widgets/ImgMediaProcessImagesOptions.cpp
deleted file mode 100644
diff --git a/Engine/Plugins/Media/ImgMedia/Source/ImgMediaEditor/Private/Widgets/ImgMediaProcessImagesOptions.h b/Engine/Plugins/Media/ImgMedia/Source/ImgMediaEditor/Private/Widgets/ImgMediaProcessImagesOptions.h
deleted file mode 100644
diff --git a/Engine/Plugins/Media/ImgMedia/Source/ImgMediaEditor/Private/Widgets/SImgMediaProcessEXR.cpp b/Engine/Plugins/Media/ImgMedia/Source/ImgMediaEditor/Private/Widgets/SImgMediaProcessEXR.cpp
new file mode 100644
--- /dev/null
+++ b/Engine/Plugins/Media/ImgMedia/Source/ImgMediaEditor/Private/Widgets/SImgMediaProcessEXR.cpp
@@ -0,0 +1,1072 @@
+// Copyright Epic Games, Inc. All Rights Reserved.
+
+#include "Widgets/SImgMediaProcessEXR.h"
+
+#include "Async/Async.h"
+#include "Customizations/ImgMediaFilePathCustomization.h"
+#include "Editor.h"
+#include "Styling/AppStyle.h"
+#include "Engine/Canvas.h"
+#include "Engine/TextureRenderTarget2D.h"
+#include "Framework/Notifications/NotificationManager.h"
+#include "HAL/PlatformFileManager.h"
+#include "IImageWrapperModule.h"
+#include "IImgMediaModule.h"
+#include "ImageUtils.h"
+#include "ImageWrapperHelper.h"
+#include "ImgMediaEditorModule.h"
+#include "ImgMediaProcessEXROptions.h"
+#include "Math/UnrealMathUtility.h"
+#include "MediaPlayer.h"
+#include "MediaSource.h"
+#include "MediaTexture.h"
+#include "Misc/FileHelper.h"
+#include "Misc/Paths.h"
+#include "Modules/ModuleManager.h"
+#include "ProfilingDebugging/CpuProfilerTrace.h"
+#include "PropertyEditorModule.h"
+#include "SlateOptMacros.h"
+#include "UObject/ObjectMacros.h"
+#include "Widgets/Input/SButton.h"
+#include "Widgets/Layout/SBox.h"
+#include "Widgets/Layout/SScrollBox.h"
+#include "Widgets/Layout/SSplitter.h"
+#include "Widgets/Notifications/SNotificationList.h"
+#include "Widgets/Text/STextBlock.h"
+
+#if IMGMEDIAEDITOR_EXR_SUPPORTED_PLATFORM
+#include "OpenExrWrapper.h"
+#endif // IMGMEDIAEDITOR_EXR_SUPPORTED_PLATFORM
+
+#define LOCTEXT_NAMESPACE "ImgMediaProcessImages"
+
+SImgMediaProcessEXR::~SImgMediaProcessEXR()
+{
+	CleanUp();
+}
+
+BEGIN_SLATE_FUNCTION_BUILD_OPTIMIZATION
+void SImgMediaProcessEXR::Construct(const FArguments& InArgs)
+{
+	// Set up widgets.
+	TSharedPtr<SBox> DetailsViewBox;
+	
+	ChildSlot
+	[
+		SNew(SScrollBox)
+
+		// Add details view.
+		+ SScrollBox::Slot()
+			[
+				SAssignNew(DetailsViewBox, SBox)
+			]
+			
+		// Add process images button.
+		+ SScrollBox::Slot()
+			.Padding(8.0f)
+			.VAlign(VAlign_Bottom)
+			.HAlign(HAlign_Right)
+			[
+				SNew(SHorizontalBox)
+
+				+ SHorizontalBox::Slot()
+					.AutoWidth()
+					[
+						SAssignNew(StartButton, SButton)
+							.ButtonStyle(&FAppStyle::Get().GetWidgetStyle<FButtonStyle>("PrimaryButton"))
+							.OnClicked(this, &SImgMediaProcessEXR::OnProcessImagesClicked)
+							.ForegroundColor(FSlateColor::UseStyle())
+							.Text(LOCTEXT("StartProcessImages", "Process Images"))
+							.ToolTipText(LOCTEXT("StartProcesssImagesButtonToolTip", "Start processing images."))
+					]
+
+				+ SHorizontalBox::Slot()
+					.AutoWidth()
+					[
+						SAssignNew(CancelButton, SButton)
+							.OnClicked(this, &SImgMediaProcessEXR::OnCancelClicked)
+							.Text(LOCTEXT("CancelProcessImages", "Cancel"))
+							.ToolTipText(LOCTEXT("CancelProcesssImagesButtonToolTip", "Cancel processing images."))
+					]
+			]
+	];
+	bIsProcessing = false;
+	bIsCancelling = false;
+	UpdateWidgets();
+
+	// Create object with our options.
+	Options = TStrongObjectPtr<UImgMediaProcessEXROptions>(NewObject<UImgMediaProcessEXROptions>(GetTransientPackage(), NAME_None));
+
+	// Create detail view with our options.
+	FPropertyEditorModule& PropertyEditorModule = FModuleManager::GetModuleChecked<FPropertyEditorModule>("PropertyEditor");
+	FDetailsViewArgs DetailsViewArgs;
+	DetailsViewArgs.bAllowSearch = false;
+	DetailsViewArgs.NameAreaSettings = FDetailsViewArgs::HideNameArea;
+	DetailsView = PropertyEditorModule.CreateDetailView(DetailsViewArgs);
+	DetailsView->RegisterInstancedCustomPropertyTypeLayout(FName(TEXT("FilePath")),
+		FOnGetPropertyTypeCustomizationInstance::CreateStatic(FImgMediaFilePathCustomization::MakeInstance));
+	DetailsView->SetObject(Options.Get());
+
+	DetailsViewBox->SetContent(DetailsView->AsShared());
+}
+END_SLATE_FUNCTION_BUILD_OPTIMIZATION
+
+void SImgMediaProcessEXR::Tick(const FGeometry& AllottedGeometry, const double InCurrentTime, const float InDeltaTime)
+{
+	SCompoundWidget::Tick(AllottedGeometry, InCurrentTime, InDeltaTime);
+
+	if (bUsePlayer)
+	{
+		HandleProcessing();
+	}
+}
+
+void SImgMediaProcessEXR::UpdateWidgets()
+{
+	StartButton->SetEnabled(!bIsProcessing);
+	CancelButton->SetEnabled(bIsProcessing && (!bIsCancelling));
+}
+
+FReply SImgMediaProcessEXR::OnProcessImagesClicked()
+{
+	if (bIsProcessing == false)
+	{
+		// Set that we are processing now.
+		bIsProcessing = true;
+		bUsePlayer = Options->bUsePlayer;
+		UpdateWidgets();
+
+		// Create notification.
+		FNotificationInfo Info(FText::GetEmpty());
+		Info.bFireAndForget = false;
+
+		ConfirmNotification = FSlateNotificationManager::Get().AddNotification(Info);
+
+		// Are we using a player?
+		if (bUsePlayer)
+		{
+			// Create player.
+			MediaPlayer = NewObject<UMediaPlayer>(GetTransientPackage(), "MediaPlayer", RF_Transient);
+			MediaPlayer->SetLooping(true);
+			MediaPlayer->PlayOnOpen = true;
+			MediaPlayer->AddToRoot();
+
+			// Create texture.
+			MediaTexture = NewObject<UMediaTexture>(GetTransientPackage(), "MediaTexture", RF_Transient);
+			MediaTexture->SetMediaPlayer(MediaPlayer);
+			MediaTexture->UpdateResource();
+			MediaTexture->AddToRoot();
+
+			// Create media source.
+			MediaSource = UMediaSource::SpawnMediaSourceForString(Options->InputPath.FilePath, GetTransientPackage());
+			if (MediaSource == nullptr)
+			{
+				return FReply::Handled();
+			}
+			MediaSource->AddToRoot();
+
+			// Start playing.
+			CurrentFrameIndex = 0;
+			CurrentTime = FTimespan::FromSeconds(0.0f);
+			MediaPlayer->SetBlockOnTimeRange(TRange<FTimespan>(CurrentTime,
+				CurrentTime + FTimespan::FromSeconds(1.0f / 100000.0f)));
+			MediaPlayer->OpenSource(MediaSource);
+		}
+		else
+		{
+			// Start async task to process files.
+			Async(EAsyncExecution::Thread, [this]()
+			{
+				ProcessAllImages();
+			});
+		}
+	}
+
+	return FReply::Handled();
+}
+
+FReply SImgMediaProcessEXR::OnCancelClicked()
+{
+	if (bIsProcessing)
+	{
+		bIsCancelling = true;
+		UpdateWidgets();
+	}
+
+	return FReply::Handled();
+}
+
+void SImgMediaProcessEXR::ProcessAllImages()
+{
+	TRACE_CPUPROFILER_EVENT_SCOPE(SImgMediaProcessEXR::ProcessAllImages);
+
+	int32 InTileWidth = Options->bEnableTiling ? Options->TileSizeX : 0;
+	int32 InTileHeight = Options->bEnableTiling ? Options->TileSizeY : 0;
+	int32 TileBorder = 0; // Note: virtual texture support is shelved for now.
+	bool bEnableMips = Options->bEnableMipMapping;
+
+	// Create output directory.
+	FString OutPath = Options->OutputPath.Path;
+	IPlatformFile& PlatformFile = FPlatformFileManager::Get().GetPlatformFile();
+	PlatformFile.CreateDirectoryTree(*OutPath);
+
+	// Get source files.
+	FString SequencePath = FPaths::GetPath(Options->InputPath.FilePath);
+	
+	TArray<FString> FoundFiles;
+	IFileManager::Get().FindFiles(FoundFiles, *SequencePath, TEXT("*"));
+	FoundFiles.RemoveAll([](const FString& FileName) { return FileName.StartsWith(TEXT(".")); });
+	FoundFiles.Sort();
+	UE_LOG(LogImgMediaEditor, Warning, TEXT("Found %i image files in %s to import."), FoundFiles.Num(), *SequencePath);
+	if (FoundFiles.Num() == 0)
+	{
+		UE_LOG(LogImgMediaEditor, Error, TEXT("No files to import."));
+	}
+	else
+	{
+		// Create image wrapper
+		FString Ext = FPaths::GetExtension(FoundFiles[0]);
+		EImageFormat ImageFormat = ImageWrapperHelper::GetImageFormat(Ext);
+
+		if (ImageFormat == EImageFormat::Invalid)
+		{
+			UE_LOG(LogImgMediaEditor, Error, TEXT("Invalid file format %s"), *Ext);
+		}
+		else
+		{
+			IImageWrapperModule& ImageWrapperModule = FModuleManager::LoadModuleChecked<IImageWrapperModule>("ImageWrapper");
+			
+			// ImageWrapper is always returning an alpha channel for RGB, so check if we really have one.
+			bool bHasAlphaChannel = HasAlphaChannel(Ext, FPaths::Combine(SequencePath, FoundFiles[0]));
+
+			// Get number of threads to use.
+			int32 NumThreads = Options->NumThreads;
+			if (NumThreads <= 0)
+			{
+				NumThreads = 8;
+			}
+
+			// Loop through all files.
+			int NumDone = 0;
+			int TotalNum = FoundFiles.Num();
+			std::atomic<int32> NumActive = 0;
+			TSharedPtr<SNotificationItem> LocalConfirmNotification = ConfirmNotification;
+			for (const FString& FileName : FoundFiles)
+			{
+				// Wait for threads to finish if we have to many.
+				while (NumActive >= NumThreads)
+				{
+					FPlatformProcess::Sleep(0.1f);
+				}
+				NumActive++;
+				
+				// Update notification with current status.
+				Async(EAsyncExecution::TaskGraphMainThread, [LocalConfirmNotification, NumDone, TotalNum]()
+				{
+					if (LocalConfirmNotification.IsValid())
+					{
+						LocalConfirmNotification->SetText(
+							FText::Format(LOCTEXT("ImgMediaCompleted", "ImgMedia Completed {0}/{1}"),
+								FText::AsNumber(NumDone), FText::AsNumber(TotalNum)));
+					}
+				});
+				NumDone++;
+
+				TSharedPtr<IImageWrapper> ImageWrapper = ImageWrapperModule.CreateImageWrapper(ImageFormat);
+				Async(EAsyncExecution::Thread, [this, SequencePath, FileName, OutPath,
+					Ext, &NumActive,
+					ImageWrapper, InTileWidth, InTileHeight,
+					TileBorder, bEnableMips, bHasAlphaChannel]() mutable
+				{
+					FString FullFileName = FPaths::Combine(SequencePath, FileName);
+
+					// Load image into buffer.
+					TArray64<uint8> InputBuffer;
+					if (!FFileHelper::LoadFileToArray(InputBuffer, *FullFileName))
+					{
+						UE_LOG(LogImgMediaEditor, Error, TEXT("Failed to load %s"), *FullFileName);
+						NumActive--;
+						return;
+					}
+					if (!ImageWrapper.IsValid() || !ImageWrapper->SetCompressed(InputBuffer.GetData(), InputBuffer.Num()))
+					{
+						UE_LOG(LogImgMediaEditor, Error, TEXT("Failed to create image wrapper for %s"), *FullFileName);
+						NumActive--;
+						return;
+					}
+
+					// Import this image.
+					FString Name = FPaths::Combine(OutPath, FileName);
+					ProcessImageCustom(ImageWrapper, InTileWidth, InTileHeight, TileBorder,
+						bEnableMips, bHasAlphaChannel, Name, false /*bUseCustomFormat*/);
+
+					NumActive--;
+				});
+
+				// Do we want to cancel?
+				if (bIsCancelling)
+				{
+					break;
+				}
+			}
+
+			// Wait for all our tasks to finish.
+			while (NumActive > 0)
+			{
+				FPlatformProcess::Sleep(0.2f);
+			}
+		}
+	}
+
+	// Close notification. Must be run on the main thread.
+	Async(EAsyncExecution::TaskGraphMainThread, [this]()
+	{
+		if (ConfirmNotification.IsValid())
+		{
+			ConfirmNotification->SetEnabled(false);
+			ConfirmNotification->SetCompletionState(bIsCancelling ? SNotificationItem::CS_Fail : SNotificationItem::CS_Success);
+			ConfirmNotification->ExpireAndFadeout();
+		}
+
+		// Done with processing.
+		bIsProcessing = false;
+		bIsCancelling = false;
+		UpdateWidgets();
+	});
+}
+
+bool SImgMediaProcessEXR::HasAlphaChannel(const FString& Ext, const FString& File)
+{
+	bool bHasAlpha = true;
+	// We just support EXR at the moment.
+	if (Ext == TEXT("exr"))
+	{
+		FRgbaInputFile InputFile(File);
+		bHasAlpha = InputFile.GetNumChannels() == 4;
+	}
+
+	return bHasAlpha;
+}
+
+void SImgMediaProcessEXR::ProcessImageCustom(TSharedPtr<IImageWrapper>& InImageWrapper,
+	int32 InTileWidth, int32 InTileHeight, int32 InTileBorder, bool bInEnableMips,
+	bool bHasAlphaChannel, const FString& InName,
+	bool bIsCustomFormat)
+{
+#if IMGMEDIAEDITOR_EXR_SUPPORTED_PLATFORM
+	TRACE_CPUPROFILER_EVENT_SCOPE(SImgMediaProcessEXR::ProcessImageCustom);
+	// Get image data.
+	ERGBFormat Format = InImageWrapper->GetFormat();
+	int32 Width = InImageWrapper->GetWidth();
+	int32 Height = InImageWrapper->GetHeight();
+	int32 BitDepth = InImageWrapper->GetBitDepth();
+	TArray64<uint8> RawData;
+	{
+		TRACE_CPUPROFILER_EVENT_SCOPE(SImgMediaProcessEXR::ProcessImageCustom:GetRaw);
+		InImageWrapper->GetRaw(Format, BitDepth, RawData);
+	}
+	ProcessImageCustomRawData(RawData, Width, Height, BitDepth,
+		InTileWidth, InTileHeight, InTileBorder, bInEnableMips,
+		bHasAlphaChannel, InName, bIsCustomFormat);
+#else // IMGMEDIAEDITOR_EXR_SUPPORTED_PLATFORM
+	UE_LOG(LogImgMediaEditor, Error, TEXT("EXR not supported on this platform."));
+#endif // IMGMEDIAEDITOR_EXR_SUPPORTED_PLATFORM
+}
+
+void SImgMediaProcessEXR::ProcessImageCustomRawData(TArray64<uint8>& RawData,
+	int32 Width, int32 Height, int32 BitDepth,
+	int32 InTileWidth, int32 InTileHeight, int32 InTileBorder, bool bInEnableMips,
+	bool bHasAlphaChannel, const FString& InName,
+	bool bIsCustomFormat)
+{
+#if IMGMEDIAEDITOR_EXR_SUPPORTED_PLATFORM
+	TRACE_CPUPROFILER_EVENT_SCOPE(SImgMediaProcessEXR::ProcessImageCustomRawData);
+	int32 DestWidth = Width;
+	int32 DestHeight = Height;
+	int32 NumTilesX = InTileWidth > 0 ? (Width + InTileWidth - 1) / InTileWidth : 1;
+	int32 NumTilesY = InTileHeight > 0 ? (Height + InTileHeight - 1) / InTileHeight : 1;
+	int32 TileWidth = InTileWidth;
+	int32 TileHeight = InTileHeight;
+	int32 BytesPerPixel = RawData.Num() / (Width * Height);
+	int32 BytesPerPixelPerChannel = BitDepth / 8;
+	int32 NumChannels = BytesPerPixel / BytesPerPixelPerChannel;
+	int32 DestNumChannels = NumChannels;
+
+	// We use 16 bits only, so convert if we have 32.
+	if (BitDepth == 32)
+	{
+		ConvertTo16Bit(RawData);
+		BitDepth = 16;
+		BytesPerPixel /= 2;
+		BytesPerPixelPerChannel /= 2;
+	}
+	else if (BitDepth != 16)
+	{
+		UE_LOG(LogImgMediaEditor, Error, TEXT("Unsupported bit depth %d"), BitDepth);
+	}
+
+	// ImageWrapper always returns an alpha channel, so make sure we really have one.
+	if ((DestNumChannels == 4) && (bHasAlphaChannel == false))
+	{
+		// Remove the alpha channel as its not needed.
+		RemoveAlphaChannel(RawData);
+		NumChannels = 3;
+		DestNumChannels = 3;
+		BytesPerPixel = BytesPerPixelPerChannel * NumChannels;
+	}
+
+	TArray64<uint8> TileBuffer;
+	TArray64<uint8> TintBuffer;
+	bool bIsTiled = (NumTilesX > 1) || (NumTilesY > 1);
+	if (bIsTiled)
+	{
+		// Take border into account.
+		DestWidth = Width + InTileBorder * 2 * NumTilesX;
+		DestHeight = Height + InTileBorder * 2 * NumTilesY;
+	}
+
+	uint8* RawDataPtr = RawData.GetData();
+
+	// Names for our channels.
+	const FString RChannelName = FString(TEXT("R"));
+	const FString GChannelName = FString(TEXT("G"));
+	const FString BChannelName = FString(TEXT("B"));
+	const FString AChannelName = FString(TEXT("A"));
+
+	FIntPoint Stride(2, 0);
+
+	// Create tiled exr file.
+	FTiledOutputFile OutFile(FIntPoint(0, 0), FIntPoint(DestWidth - 1, DestHeight - 1),
+		FIntPoint(0, 0), FIntPoint(DestWidth - 1, DestHeight - 1), bIsTiled || bIsCustomFormat);
+
+	// Add attributes.
+	if (bIsCustomFormat)
+	{
+		OutFile.AddIntAttribute(IImgMediaModule::CustomFormatAttributeName.Resolve().ToString(), 1);
+
+		// These attributes will not be added and therefore not found by EXR reader if it is not tiled.
+		if (bIsTiled)
+		{
+			OutFile.AddIntAttribute(IImgMediaModule::CustomFormatTileWidthAttributeName.Resolve().ToString(), TileWidth);
+			OutFile.AddIntAttribute(IImgMediaModule::CustomFormatTileHeightAttributeName.Resolve().ToString(), TileHeight);
+			OutFile.AddIntAttribute(IImgMediaModule::CustomFormatTileBorderAttributeName.Resolve().ToString(), InTileBorder);
+		}
+	}
+
+	// Add channels.
+	if (DestNumChannels == 4)
+	{
+		OutFile.AddChannel(AChannelName);
+	}
+	if (DestNumChannels >= 3)
+	{
+		OutFile.AddChannel(BChannelName);
+		OutFile.AddChannel(GChannelName);
+		OutFile.AddChannel(RChannelName);
+	}
+
+	// Create output.
+	if (bIsCustomFormat)
+	{
+		OutFile.CreateOutputFile(InName, DestWidth, DestHeight, bInEnableMips, 1);
+	}
+	else
+	{
+		OutFile.CreateOutputFile(InName, TileWidth, TileHeight, bInEnableMips, 1);
+	}
+
+	if (DestNumChannels == 4)
+	{
+		OutFile.AddFrameBufferChannel(AChannelName, nullptr, Stride);
+	}
+	if (DestNumChannels >= 3)
+	{
+		OutFile.AddFrameBufferChannel(BChannelName, nullptr, Stride);
+		OutFile.AddFrameBufferChannel(GChannelName, nullptr, Stride);
+		OutFile.AddFrameBufferChannel(RChannelName, nullptr, Stride);
+	}
+
+	// Flip between 2 buffers making mips.
+	TArray64<uint8> RawData2;
+	uint8* MipBuffer[2];
+	MipBuffer[0] = RawDataPtr;
+	MipBuffer[1] = nullptr;
+	int32 CurrentMipBufferIndex = 0;
+
+	// Loop over each mip level.
+	TRACE_CPUPROFILER_EVENT_SCOPE(SImgMediaProcessEXR::ProcessImageCustom:CreateMips);
+	int32 NumMips = OutFile.GetNumberOfMipLevels();
+	int32 MipSourceWidth = Width;
+	int32 MipSourceHeight = Height;
+	int32 LastMipWidth = 0;
+	int32 LastMipHeight = 0;
+	for (int32 MipLevel = 0; MipLevel < NumMips; MipLevel++)
+	{
+		int32 MipWidth = OutFile.GetMipWidth(MipLevel);
+		int32 MipHeight = OutFile.GetMipHeight(MipLevel);
+		uint8* CurrentBuffer = MipBuffer[CurrentMipBufferIndex];
+		uint8* LastBuffer = MipBuffer[CurrentMipBufferIndex ^ 1];
+
+		// Allocate space for the other buffer.
+		if (CurrentBuffer == nullptr)
+		{
+			RawData2.AddUninitialized(MipWidth* MipHeight* BytesPerPixel);
+			MipBuffer[CurrentMipBufferIndex] = RawData2.GetData();
+			CurrentBuffer = MipBuffer[CurrentMipBufferIndex];
+		}
+
+		// Generate mip data.
+		if (MipLevel != 0)
+		{
+			TRACE_CPUPROFILER_EVENT_SCOPE(SImgMediaProcessEXR::ProcessImageCustom:GenerateMipData);
+			int32 SourceStrideX = NumChannels;
+			int32 SourceStrideY = LastMipWidth * NumChannels;
+			for (int32 PixelY = 0; PixelY < MipHeight; ++PixelY)
+			{
+				for (int32 PixelX = 0; PixelX < MipWidth; ++PixelX)
+				{
+					int32 PixelOffset = (PixelX + PixelY * MipWidth) * NumChannels;
+					for (int32 Channel = 0; Channel < NumChannels; ++Channel)
+					{
+						// Box filter.
+						int32 SourceWidth = LastMipWidth;
+						int32 SourcePixelOffset = (PixelX + PixelY * SourceWidth) * NumChannels * 2 +
+							Channel;
+						FFloat16 SourcePixels[4];
+						SourcePixels[0] =
+							((FFloat16*)LastBuffer)[SourcePixelOffset];
+						SourcePixels[1] =
+							((FFloat16*)LastBuffer)[SourcePixelOffset + SourceStrideX];
+						SourcePixels[2] =
+							((FFloat16*)LastBuffer)[SourcePixelOffset + SourceStrideY];
+						SourcePixels[3] =
+							((FFloat16*)LastBuffer)[SourcePixelOffset + SourceStrideX + SourceStrideY];
+
+						((FFloat16*)CurrentBuffer)[PixelOffset + Channel] =
+							(SourcePixels[0] + SourcePixels[1] + SourcePixels[2] + SourcePixels[3]) * 0.25f;
+					}
+				}
+			}
+		}
+
+		// Tint mip levels?
+		if (Options->bEnableMipLevelTint)
+		{
+			TintData(CurrentBuffer, TintBuffer,
+				MipLevel, MipWidth, MipHeight,
+				NumChannels);
+			CurrentBuffer = TintBuffer.GetData();
+		}
+
+		// Do we need to tile this mip?
+		// Need to also check that this is actually a valid mip level.
+		if ((bIsCustomFormat) && (bIsTiled) && (MipSourceWidth > 0) && (MipSourceHeight > 0))
+		{
+			int32 MipTileWidth = TileWidth;
+			int32 MipTileHeight = TileHeight;
+
+			// A tile could be larger than the mip level when dealing with mips.
+			if (MipTileWidth > MipSourceWidth)
+			{
+				MipTileWidth = MipSourceWidth;
+			}
+			if (MipTileHeight > MipSourceHeight)
+			{
+				MipTileHeight = MipSourceHeight;
+			}
+
+			int32 OutputWidth = 0;
+			int32 OutputHeight = 0;
+			int32 MipNumTilesX = (MipSourceWidth + MipTileWidth - 1) / MipTileWidth;
+			int32 MipNumTilesY = (MipSourceHeight + MipTileHeight - 1) / MipTileHeight;
+
+			// Make sure our sizes match the mip size we get from EXR.
+			int32 ExpectedMipWidth = MipSourceWidth + MipNumTilesX * InTileBorder * 2;
+			if (ExpectedMipWidth != MipWidth)
+			{
+				UE_LOG(LogImgMediaEditor, Error,
+					TEXT("Expected mip level width of %d, but got %d (SourceWidth:%d NumTiles:%d TileBorder:%d"),
+					ExpectedMipWidth, MipHeight,
+					MipSourceWidth, MipNumTilesX, InTileBorder);
+			}
+			int32 ExpectedMipHeight = MipSourceHeight + MipNumTilesY * InTileBorder * 2;
+			if (ExpectedMipHeight != MipHeight)
+			{
+				UE_LOG(LogImgMediaEditor, Error,
+					TEXT("Expected mip level height of %d, but got %d (SourceHeight:%d NumTiles:%d TileBorder:%d"),
+					ExpectedMipHeight, MipHeight,
+					MipSourceHeight, MipNumTilesY, InTileBorder);
+			}
+
+			// Tile the buffer.
+			TileData(CurrentBuffer, TileBuffer,
+				MipSourceWidth, MipSourceHeight, MipWidth, MipHeight,
+				MipNumTilesX, MipNumTilesY,
+				MipTileWidth, MipTileHeight, InTileBorder,
+				BytesPerPixel);
+			CurrentBuffer = TileBuffer.GetData();
+		}
+
+		// Write to EXR.
+		TRACE_CPUPROFILER_EVENT_SCOPE(SImgMediaProcessEXR::ProcessImageCustom:WriteEXR);
+		if (bIsCustomFormat)
+		{
+			Stride.Y = MipWidth * BytesPerPixel;
+			int64 BufferOffset = 0;
+			int64 SingleBufferOffset = MipWidth * BytesPerPixelPerChannel;
+			if (DestNumChannels == 4)
+			{
+				OutFile.UpdateFrameBufferChannel(AChannelName, CurrentBuffer, Stride);
+				BufferOffset += SingleBufferOffset;
+			}
+
+			OutFile.UpdateFrameBufferChannel(BChannelName, CurrentBuffer + BufferOffset, Stride);
+			BufferOffset += SingleBufferOffset;
+			OutFile.UpdateFrameBufferChannel(GChannelName, CurrentBuffer + BufferOffset, Stride);
+			BufferOffset += SingleBufferOffset;
+			OutFile.UpdateFrameBufferChannel(RChannelName, CurrentBuffer + BufferOffset, Stride);
+			BufferOffset += SingleBufferOffset;
+		}
+		else
+		{
+			Stride.X = BytesPerPixel;
+			Stride.Y = MipWidth * BytesPerPixel;
+			int64 BufferOffset = 0;
+			int64 SingleBufferOffset = BytesPerPixelPerChannel;
+			
+
+			OutFile.UpdateFrameBufferChannel(RChannelName, CurrentBuffer + BufferOffset, Stride);
+			BufferOffset += SingleBufferOffset;
+			OutFile.UpdateFrameBufferChannel(GChannelName, CurrentBuffer + BufferOffset, Stride);
+			BufferOffset += SingleBufferOffset;
+			OutFile.UpdateFrameBufferChannel(BChannelName, CurrentBuffer + BufferOffset, Stride);
+			BufferOffset += SingleBufferOffset;
+
+			if (DestNumChannels == 4)
+			{
+				OutFile.UpdateFrameBufferChannel(AChannelName, CurrentBuffer + BufferOffset, Stride);
+			}
+		}
+
+		OutFile.SetFrameBuffer();
+
+		if (bIsCustomFormat)
+		{
+			OutFile.WriteTile(0, 0, MipLevel);
+		}
+		else
+		{
+			OutFile.WriteTiles(0, OutFile.GetNumXTiles(MipLevel) - 1, 0, OutFile.GetNumYTiles(MipLevel) - 1, MipLevel);
+		}
+
+		// Switch buffers.
+		CurrentMipBufferIndex ^= 1;
+		MipSourceHeight /= 2;
+		MipSourceWidth /= 2;
+		LastMipWidth = MipWidth;
+		LastMipHeight = MipHeight;
+	}
+
+#else // IMGMEDIAEDITOR_EXR_SUPPORTED_PLATFORM
+	UE_LOG(LogImgMediaEditor, Error, TEXT("EXR not supported on this platform."));
+#endif // IMGMEDIAEDITOR_EXR_SUPPORTED_PLATFORM
+}
+
+
+void SImgMediaProcessEXR::ConvertTo16Bit(TArray64<uint8>& Buffer)
+{
+	TRACE_CPUPROFILER_EVENT_SCOPE(SImgMediaProcessEXR::ConvertTo16Bit);
+
+	int32 BytesPerPixelPerChannel = 4;
+	int64 BufferSize = Buffer.Num() / BytesPerPixelPerChannel;
+	FFloat32* InBufferPtr = (FFloat32*)(Buffer.GetData());
+	FFloat16* OutBufferPtr = (FFloat16*)(Buffer.GetData());
+
+	// Loop through the buffer.
+	for (int64 Index = 0; Index < BufferSize; ++Index)
+	{
+		// Convert the data in place.
+		OutBufferPtr[Index] = InBufferPtr[Index].FloatValue;
+	}
+
+	// Don't bother shrinking as its just a waste and extra work.
+	Buffer.SetNum(BufferSize / 2, false);
+}
+
+void SImgMediaProcessEXR::RemoveAlphaChannel(TArray64<uint8>& Buffer)
+{
+	TRACE_CPUPROFILER_EVENT_SCOPE(SImgMediaProcessEXR::RemoveAlphaChannel);
+
+	int32 BytesPerPixelPerChannel = 2;
+	int64 BufferSize = Buffer.Num() / BytesPerPixelPerChannel;
+	uint16* BufferPtr = (uint16*)(Buffer.GetData());
+	
+	// Loop through the buffer.
+	int64 OutIndex = 0;
+	for (int64 Index = 0; Index < BufferSize; ++Index)
+	{
+		// Skip every fourth channel (i.e. the alpha channel).
+		if ((Index & 0x3) != 3)
+		{
+			// Copy the data in place.
+			BufferPtr[OutIndex] = BufferPtr[Index];
+			OutIndex++;
+		}
+	}
+
+	// Don't bother shrinking as its just a waste and extra work.
+	Buffer.SetNum((BufferSize * 3) / 4, false);
+}
+
+void SImgMediaProcessEXR::TintData(uint8* SourceData, TArray64<uint8>& DestArray,
+	int32 InMipLevel, int32 InWidth, int32 InHeight, int32 InNumChannels)
+{
+	TRACE_CPUPROFILER_EVENT_SCOPE(SImgMediaProcessEXR::TintData);
+
+	// Get tint colour.
+	FLinearColor TintColor = FLinearColor::White;
+	if (Options->MipLevelTints.Num() > 0)
+	{
+		TintColor = Options->MipLevelTints[InMipLevel % Options->MipLevelTints.Num()];
+	}
+
+	// Set up destination buffer.
+	DestArray.Reset();
+	DestArray.AddUninitialized(InWidth * InHeight * InNumChannels * 2);
+	FFloat16* Buffer = (FFloat16*)SourceData;
+	FFloat16* OutBuffer = (FFloat16*)DestArray.GetData();
+		
+	// Tint buffer.
+	for (int32 PixelY = 0; PixelY < InHeight; ++PixelY)
+	{
+		for (int32 PixelX = 0; PixelX < InWidth; ++PixelX)
+		{
+			OutBuffer[0] = (Buffer[0].GetFloat() + TintColor.R) * 0.5f;
+			OutBuffer[1] = (Buffer[1].GetFloat() + TintColor.G) * 0.5f;
+			OutBuffer[2] = (Buffer[2].GetFloat() + TintColor.B) * 0.5f;
+			if (InNumChannels == 4)
+			{
+				OutBuffer[3] = Buffer[3];
+			}
+				
+			Buffer += InNumChannels;
+			OutBuffer += InNumChannels;
+		}
+	}
+}
+
+void SImgMediaProcessEXR::TileData(uint8* SourceData, TArray64<uint8>& DestArray,
+	int32 SourceWidth, int32 SourceHeight, int32 DestWidth, int32 DestHeight,
+	int32 NumTilesX, int32 NumTilesY,
+	int32 TileWidth, int32 TileHeight, int32 InTileBorder,
+	int32 BytesPerPixel)
+{
+	TRACE_CPUPROFILER_EVENT_SCOPE(SImgMediaProcessEXR::TileData);
+
+	// We don't support tile borders larger than a tile size,
+	// but this shuld not happen in practice.
+	if ((InTileBorder > TileWidth) || (InTileBorder > TileHeight))
+	{
+		UE_LOG(LogImgMediaEditor, Error, TEXT("Tile border is larger than tile size. Clamping to tile size."));
+		InTileBorder = FMath::Min(TileWidth, TileHeight);
+	}
+
+	// Set up destination buffer.
+	DestArray.Reset();
+	DestArray.AddUninitialized(DestWidth * DestHeight * BytesPerPixel);
+
+	uint8* DestData = DestArray.GetData();
+	int32 DestTileWidth = TileWidth + InTileBorder * 2;
+	int32 DestTileHeight = TileHeight + InTileBorder * 2;
+
+	// Make sure our output tile size is not bigger than the output size.
+	if ((DestTileWidth > DestWidth) || (DestTileHeight > DestHeight))
+	{
+		// This is not a valid mip level, so just ignore.
+		return;
+	}
+
+	int32 BytesPerTile = TileWidth * TileHeight * BytesPerPixel;
+	int32 ByterPerDestTile = DestTileWidth * DestTileHeight * BytesPerPixel;
+	uint8* DestTile = DestData;
+
+	// Loop over y tiles.
+	for (int32 TileY = 0; TileY < NumTilesY; ++TileY)
+	{
+		// Loop over x tiles.
+		for (int32 TileX = 0; TileX < NumTilesX; ++TileX)
+		{
+			// Get address of the source and destination tiles.
+			uint8* SourceTile = SourceData +
+				(TileX * TileWidth + TileY * SourceWidth * TileHeight) * BytesPerPixel;
+			
+			// If this tile is over the right edge of our image, then make this tile smaller
+			// so it does not exceed the image size.
+			int32 NumberOfPixelsToCopy = TileWidth;
+			int32 ThisDestTileWidth = DestTileWidth;
+			if ((TileX + 1) * TileWidth > SourceWidth)
+			{
+				NumberOfPixelsToCopy = SourceWidth - TileX * TileWidth;
+				ThisDestTileWidth -= TileWidth - NumberOfPixelsToCopy;
+			}
+
+			// Create a left border.
+			int32 DestTileOffset = 0;
+			if (TileX > 0)
+			{
+				NumberOfPixelsToCopy += InTileBorder;
+				// Offset the source to get the extra pixels.
+				SourceTile -= InTileBorder * BytesPerPixel;
+			}
+			else
+			{
+				// Offset the destination as we are skipping this border as we have no data.
+				DestTileOffset = InTileBorder * BytesPerPixel;
+				DestTile += DestTileOffset;
+			}
+
+			// Create a right border.
+			if (TileX < NumTilesX - 1)
+			{
+				NumberOfPixelsToCopy += InTileBorder;
+			}
+
+			// If this tile is over the bottom edge of our image, then make this tile smaller
+			// so it does not exceed the image size.
+			int32 ThisDestTileHeight = DestTileHeight;
+			if ((TileY + 1) * TileHeight > SourceHeight)
+			{
+				ThisDestTileHeight = SourceHeight - TileY * TileHeight;
+			}
+
+			// Loop over each row in the tile.
+			for (int32 Row = 0; Row < ThisDestTileHeight; ++Row)
+			{
+				// Make sure we don't go beyond the source data.
+				int32 SourceRow = Row - InTileBorder;
+				if (TileY == 0)
+				{
+					SourceRow = FMath::Max(SourceRow, 0);
+				}
+				if (TileY == NumTilesY - 1)
+				{
+					SourceRow = FMath::Min(SourceRow, TileHeight - 1);
+				}
+
+				uint8* SourceLine = SourceTile + SourceRow * SourceWidth * BytesPerPixel;
+				uint8* DestLine = DestTile;
+
+				// Copy the main data.
+				FMemory::Memcpy(DestLine, SourceLine, NumberOfPixelsToCopy * BytesPerPixel);
+
+				// Increment our pointer to the next tile.
+				// We have to remove any DestTileOffset we applied earlier.
+				DestTile += ThisDestTileWidth * BytesPerPixel - DestTileOffset;
+			}
+		}
+	}
+}
+
+void SImgMediaProcessEXR::HandleProcessing()
+{
+	// Are we processing?
+	if (bIsProcessing)
+	{
+		TRACE_CPUPROFILER_EVENT_SCOPE(SImgMediaProcessEXR::HandleProcessing);
+		// We did not cancel yet?
+		bool bShouldExit = false;
+		if ((MediaPlayer != nullptr) && (bIsCancelling == false))
+		{
+			// Get which frame the player is on.
+			int32 PlayerFrame = 0;
+			if (FrameDuration.GetTotalSeconds() > 0.0f)
+			{
+				PlayerFrame = (int32)FMath::RoundHalfToEven(MediaPlayer->GetTime().GetTotalSeconds() / FrameDuration.GetTotalSeconds());
+			}
+
+			UE_LOG(LogImgMediaEditor, Verbose,
+				TEXT("ProcessImages Time:%f PlayerTime:%f Duration:%f Frame:%d"),
+				CurrentTime.GetTotalSeconds(),
+				MediaPlayer->GetTime().GetTotalSeconds(),
+				MediaPlayer->GetDuration().GetTotalSeconds(), PlayerFrame);
+
+			// Has the player stopped playing?
+			if (MediaPlayer->IsClosed())
+			{
+				bShouldExit = true;
+			}
+			// Is this the frame we want?
+			else if ((MediaPlayer->IsPreparing() == false) && (CurrentFrameIndex == PlayerFrame))
+			{
+				// Are we set up yet?
+				if (RenderTarget == nullptr)
+				{
+					CreateRenderTarget();
+					// Get frame duration.
+					float FrameRate = MediaPlayer->GetVideoTrackFrameRate(INDEX_NONE, INDEX_NONE);
+					if (FrameRate <= 0.0f)
+					{
+						FrameRate = 24.0f;
+					}
+					FrameDuration = FTimespan::FromSeconds(1.0f / FrameRate);
+				}
+
+				// Did we get a render target?
+				if (RenderTarget != nullptr)
+				{
+					// Copy media texture to our render target.
+					DrawTextureToRenderTarget();
+
+					// Process this render.
+					TArray64<uint8> RawData;
+					bool bReadSuccess = FImageUtils::GetRawData(RenderTarget, RawData);
+					if (bReadSuccess)
+					{
+						int32 Width = RenderTarget->GetSurfaceWidth();
+						int32 Height = RenderTarget->GetSurfaceHeight();
+						int32 BitDepth = 16;
+						int32 InTileWidth = Options->bEnableTiling ? Options->TileSizeX : 0;
+						int32 InTileHeight = Options->bEnableTiling ? Options->TileSizeY : 0;
+						int32 TileBorder = 0; // Note: virtual texture support is shelved for now.
+						bool bEnableMips = Options->bEnableMipMapping;
+						bool bHasAlphaChannel = false;
+						FString OutPath = Options->OutputPath.Path;
+						FString FileName = FString::Printf(TEXT("image%05d.exr"), CurrentFrameIndex);
+						FString Name = FPaths::Combine(OutPath, FileName);
+
+						Async(EAsyncExecution::Thread, [this, RawData = MoveTemp(RawData), Width, Height, BitDepth,
+							InTileWidth, InTileHeight, TileBorder, bEnableMips,
+							bHasAlphaChannel, Name]() mutable
+						{
+							ProcessImageCustomRawData(RawData, Width, Height, BitDepth,
+								InTileWidth, InTileHeight, TileBorder, bEnableMips,
+								bHasAlphaChannel, Name, false /*bUseCustomFormat*/);
+						});
+					}
+					else
+					{
+						UE_LOG(LogImgMediaEditor, Error, TEXT("ProcessImages failed to get raw data."));
+					}
+
+					// Update notification.
+					if (ConfirmNotification.IsValid())
+					{
+						ConfirmNotification->SetText(
+							FText::Format(LOCTEXT("ImgMediaCompleted2", "ImgMedia Completed {0}"),
+								FText::AsNumber(CurrentFrameIndex)));
+					}
+
+					// Next frame.
+					CurrentTime += FrameDuration;
+					CurrentFrameIndex++;
+					if (CurrentTime >= MediaPlayer->GetDuration())
+					{
+						bShouldExit = true;
+					}
+					else
+					{
+						MediaPlayer->SetBlockOnTimeRange(TRange<FTimespan>(CurrentTime, CurrentTime + FrameDuration));
+					}
+				}
+			}
+		}
+		else
+		{
+			bShouldExit = true;
+		}
+
+		// Are we done?
+		if (bShouldExit)
+		{
+			// Remove notification.
+			if (ConfirmNotification.IsValid())
+			{
+				ConfirmNotification->SetEnabled(false);
+				ConfirmNotification->SetCompletionState(bIsCancelling ? SNotificationItem::CS_Fail : SNotificationItem::CS_Success);
+				ConfirmNotification->ExpireAndFadeout();
+				ConfirmNotification.Reset();
+			}
+
+			// Clean up.
+			bIsCancelling = false;
+			bIsProcessing = false;
+			CleanUp();
+			UpdateWidgets();
+		}
+	}
+}
+
+void SImgMediaProcessEXR::CreateRenderTarget()
+{
+	FIntPoint VideoDim = MediaPlayer->GetVideoTrackDimensions(INDEX_NONE, INDEX_NONE);
+	if (VideoDim.X > 0)
+	{
+		int32 Width = VideoDim.X;
+		int32 Height = VideoDim.Y;
+
+		RenderTarget = NewObject<UTextureRenderTarget2D>(GetTransientPackage(), TEXT("ImgMediaProcessImages"));
+		RenderTarget->RenderTargetFormat = RTF_RGBA16f;
+		RenderTarget->InitAutoFormat(Width, Height);
+		RenderTarget->AddToRoot();
+		RenderTarget->UpdateResourceImmediate(true);
+	}
+}
+
+void SImgMediaProcessEXR::DrawTextureToRenderTarget()
+{
+	UWorld* World = nullptr;
+	World = GEditor->GetEditorWorldContext().World();
+	World->FlushDeferredParameterCollectionInstanceUpdates();
+
+	FTextureRenderTargetResource* RenderTargetResource = RenderTarget->GameThread_GetRenderTargetResource();
+
+	UCanvas* Canvas = World->GetCanvasForDrawMaterialToRenderTarget();
+	FCanvas RenderCanvas(
+		RenderTargetResource,
+		nullptr,
+		World,
+		World->FeatureLevel);
+	Canvas->Init(RenderTarget->SizeX, RenderTarget->SizeY, nullptr, &RenderCanvas);
+	Canvas->Update();
+
+	{
+		ENQUEUE_RENDER_COMMAND(FlushDeferredResourceUpdateCommand)(
+			[RenderTargetResource](FRHICommandListImmediate& RHICmdList)
+		{
+			RenderTargetResource->FlushDeferredResourceUpdate(RHICmdList);
+		});
+
+		Canvas->K2_DrawTexture(MediaTexture, FVector2D(0, 0), FVector2D(RenderTarget->SizeX, RenderTarget->SizeY), FVector2D(0, 0));
+
+		RenderCanvas.Flush_GameThread();
+		Canvas->Canvas = nullptr;
+		RenderTarget->UpdateResourceImmediate(false);
+	}
+}
+
+void SImgMediaProcessEXR::CleanUp()
+{
+	if (MediaPlayer != nullptr)
+	{
+		MediaPlayer->Close();
+		MediaPlayer->RemoveFromRoot();
+		MediaPlayer = nullptr;
+	}
+	if (MediaTexture != nullptr)
+	{
+		MediaTexture->RemoveFromRoot();
+		MediaTexture = nullptr;
+	}
+	if (MediaSource != nullptr)
+	{
+		MediaSource->RemoveFromRoot();
+		MediaSource = nullptr;
+	}
+	if (RenderTarget != nullptr)
+	{
+		RenderTarget->RemoveFromRoot();
+		RenderTarget = nullptr;
+	}
+}
+
+#undef LOCTEXT_NAMESPACE
diff --git a/Engine/Plugins/Media/ImgMedia/Source/ImgMediaEditor/Private/Widgets/SImgMediaProcessEXR.h b/Engine/Plugins/Media/ImgMedia/Source/ImgMediaEditor/Private/Widgets/SImgMediaProcessEXR.h
new file mode 100644
--- /dev/null
+++ b/Engine/Plugins/Media/ImgMedia/Source/ImgMediaEditor/Private/Widgets/SImgMediaProcessEXR.h
@@ -0,0 +1,209 @@
+// Copyright Epic Games, Inc. All Rights Reserved.
+
+#pragma once
+
+#include "CoreMinimal.h"
+#include "ImgMediaProcessEXROptions.h"
+#include "Input/Reply.h"
+#include "UObject/StrongObjectPtr.h"
+#include "Widgets/SCompoundWidget.h"
+
+class IDetailsView;
+class IImageWrapper;
+class SButton;
+class SNotificationItem;
+class UMediaPlayer;
+class UMediaSource;
+class UMediaTexture;
+class UTextureRenderTarget2D;
+
+/**
+ * SImgMediaProcessEXR provides processing of image sequences.
+ */
+class SImgMediaProcessEXR : public SCompoundWidget
+{
+public:
+	SLATE_BEGIN_ARGS(SImgMediaProcessEXR){}
+	SLATE_END_ARGS()
+
+	virtual ~SImgMediaProcessEXR();
+
+	void Construct(const FArguments& InArgs);
+
+	/** SWidget interface */
+	virtual void Tick(const FGeometry& AllottedGeometry, const double InCurrentTime, const float InDeltaTime);
+
+
+private:
+	/**
+	 * Updates our widgets based on the current state.
+	 */
+	void UpdateWidgets();
+	
+	/** Called when we click on the process images button. */
+	FReply OnProcessImagesClicked();
+	/** Called when we click on the cancel button. */
+	FReply OnCancelClicked();
+
+	/**
+	 * Processes all images in the sequence and generate tiles/mips.
+	 */
+	void ProcessAllImages();
+	
+	/**
+	  * Checks if this file has an alpha channel.
+	  *
+	  * @param Ext			File extension.
+	  * @param File			Full path to file.
+	  * @return				True if the file has an alpha channel.
+	  */
+	bool HasAlphaChannel(const FString& Ext, const FString& File);
+
+	/**
+	 * Processess a single image and writes out a file.
+	 * Tiles and mips may be generated.
+	 * This does NOT run on the game thread.
+	 *
+	 * @param InImageWrapper	ImageWrapper to read/write the image.
+	 * @param InTileWidth		Desired width of tiles.
+	 * @param InTileHeight		Desired height of tiles.
+	 * @param InTileBorder		Number of pixels to duplicate along a tile edge.
+	 * @param bInEnableMips		Turn on mip mapping.
+	 * @param bHasAlphaChannel	True if there really is an alpha channel.
+	 * @param InName			Full path and name of file to write.
+	 * @param bIsCustomFormat	True to output in our custom format.
+	 */
+	void ProcessImageCustom(TSharedPtr<IImageWrapper>& InImageWrapper,
+		int32 InTileWidth, int32 InTileHeight, int32 InTileBorder, 
+		bool bInEnableMips, bool bHasAlphaChannel, const FString& InName,
+		bool bIsCustomFormat);
+
+	/**
+	 * Processess a single image and writes out a file.
+	 * Tiles and mips may be generated.
+	 * This might not run on the game thread.
+	 *
+	 * @param RawData			Image pixel data (RGBA FFloat16).
+	 * @param Width				Width of image.
+	 * @param Height			Height of image.
+	 * @param BitDepth			Number of bits per pixel.
+	 * @param InTileWidth		Desired width of tiles.
+	 * @param InTileHeight		Desired height of tiles.
+	 * @param InTileBorder		Number of pixels to duplicate along a tile edge.
+	 * @param bInEnableMips		Turn on mip mapping.
+	 * @param bHasAlphaChannel	True if there really is an alpha channel.
+	 * @param InName			Full path and name of file to write.
+	 * @param bIsCustomFormat	True to output in our custom format.
+	 */
+	void ProcessImageCustomRawData(TArray64<uint8>& RawData,
+		int32 Width, int32 Height, int32 BitDepth,
+		int32 InTileWidth, int32 InTileHeight, int32 InTileBorder, bool bInEnableMips,
+		bool bHasAlphaChannel, const FString& InName, bool bIsCustomFormat);
+	
+	/**
+	 * Converts a buffer from 32 bit floating point to 16 bit floating point.
+	 * This will be done in place.
+	 * 
+	 * @param Buffer			Buffer to convert.
+	 */
+	void ConvertTo16Bit(TArray64<uint8>& Buffer);
+
+	/**
+	 * Removes the alpha channel from a buffer.
+	 * This will be done in place.
+	 * Assumes 2 bytes per channel, 4 channels, alpha is the last channel.
+	 * 
+	 * @param Buffer			Buffer to remove alpha from.
+	 */
+	void RemoveAlphaChannel(TArray64<uint8>& Buffer);
+
+	/**
+	 * Tints a source buffer and outputs it to a destination.
+	 *
+	 * @param SourceData		Source image.
+	 * @param DestArray			Destination image.
+	 * @param InMipLevel		The mip level of this image.
+	 * @param InWidth			Width of source in pixels.
+	 * @param InHeight			Height of source in pixels.
+	 * @param InNumChannels		Number of channels in the image.
+	 */
+	void TintData(uint8* SourceData, TArray64<uint8>& DestArray, 
+		int32 InMipLevel, int32 InWidth, int32 InHeight,
+		int32 InNumChannels);
+
+	/**
+	 * Creates tiles from a source and outputs it to a destination.
+	 *
+	 * @param SourceData		Source image.
+	 * @param DestArray			Destination image.
+	 * @param SourceWidth		Width of source in pixels.
+	 * @param SourceHeight		Height of source in pixels.
+	 * @param DestWidth			Width of destination in pixels.
+	 * @param DestHeight		Height of destination in pixels.
+	 * @param NumTilesX			Number of tiles in X direction.
+	 * @param NumTilesY			Number of tiles in Y direction.
+	 * @param TileWidth			Width of a tile (without borders) in pixels.
+	 * @param TileHeight		Height of a tile (without borders) in pixels.
+	 * @param InTileBorder		Size of border in pixels.
+	 * @param BytesPerPixel		Number of bytes per pixel.
+	 */
+	void TileData(uint8* SourceData, TArray64<uint8>& DestArray,
+		int32 SourceWidth, int32 SourceHeight, int32 DestWidth, int32 DestHeight,
+		int32 NumTilesX, int32 NumTilesY,
+		int32 TileWidth, int32 TileHeight, int32 InTileBorder,
+		int32 BytesPerPixel);
+
+	/**
+	 * Call this every frame to handle processing.
+	 */
+	void HandleProcessing();
+
+	/**
+	 * Creates a render target.
+	 */
+	void CreateRenderTarget();
+
+	/**
+	 * Draws the media texture to our render target.
+	 */
+	void DrawTextureToRenderTarget();
+
+	/**
+	 * Clean up after processing is done.
+	 */
+	void CleanUp();
+
+	/** Holds our details view. */
+	TSharedPtr<class IDetailsView> DetailsView;
+	/** Object that holds our options. */
+	TStrongObjectPtr<UImgMediaProcessEXROptions> Options;
+	/** Holds the button to start processing. */
+	TSharedPtr<SButton> StartButton;
+	/** Holds the button to cancel processing. */
+	TSharedPtr<SButton> CancelButton;
+	/** Notification to update user on our progress. */
+	TSharedPtr<SNotificationItem> ConfirmNotification;
+
+	/** Used to read in media. */
+	TObjectPtr<UMediaPlayer> MediaPlayer;
+	/** Texture for the player. */
+	TObjectPtr<UMediaTexture> MediaTexture;
+	/** Source for the player. */
+	TObjectPtr<UMediaSource> MediaSource;
+	/** Render target to render the media texture to. */
+	TObjectPtr<UTextureRenderTarget2D> RenderTarget;
+	/** The current time we are processing. */
+	FTimespan CurrentTime;
+	/** How long a frame is. */
+	FTimespan FrameDuration;
+	/** Current frame index we are processinig. */
+	int32 CurrentFrameIndex = 0;
+
+	/** True if we are currently processing. */
+	bool bIsProcessing = false;
+	/** True if we want to cancel processing. */
+	bool bIsCancelling = false;
+	/** True if we are using a player to read in the images. */
+	bool bUsePlayer = false;
+	
+};
diff --git a/Engine/Plugins/Media/ImgMedia/Source/ImgMediaEditor/Private/Widgets/SImgMediaProcessImages.cpp b/Engine/Plugins/Media/ImgMedia/Source/ImgMediaEditor/Private/Widgets/SImgMediaProcessImages.cpp
deleted file mode 100644
diff --git a/Engine/Plugins/Media/ImgMedia/Source/ImgMediaEditor/Private/Widgets/SImgMediaProcessImages.h b/Engine/Plugins/Media/ImgMedia/Source/ImgMediaEditor/Private/Widgets/SImgMediaProcessImages.h
deleted file mode 100644
