ImgMedia: Renaming editor tabs.

diff --git a/Engine/Plugins/Media/ImgMedia/Source/ImgMediaEditor/Private/ImgMediaEditorModule.cpp b/Engine/Plugins/Media/ImgMedia/Source/ImgMediaEditor/Private/ImgMediaEditorModule.cpp
--- a/Engine/Plugins/Media/ImgMedia/Source/ImgMediaEditor/Private/ImgMediaEditorModule.cpp
+++ b/Engine/Plugins/Media/ImgMedia/Source/ImgMediaEditor/Private/ImgMediaEditorModule.cpp
@@ -127,24 +127,24 @@ protected:
 		FGlobalTabmanager::Get()->RegisterNomadTabSpawner(ImgMediaBandwidthTabName,
 			FOnSpawnTab::CreateStatic(&FImgMediaEditorModule::SpawnBandwidthTab))
 			.SetGroup(MediaBrowserGroup)
-			.SetDisplayName(LOCTEXT("ImgMediaBandwidthTabTitle", "Bandwidth"))
-			.SetTooltipText(LOCTEXT("ImgMediaBandwidthTooltipText", "Open the bandwidth tab."))
+			.SetDisplayName(LOCTEXT("ImgMediaBandwidthMonitorTabTitle", "Bandwidth Monitor"))
+			.SetTooltipText(LOCTEXT("ImgMediaBandwidthMonitorTooltipText", "Open the bandwidth monitor tab."))
 			.SetIcon(FSlateIcon(FAppStyle::GetAppStyleSetName(), "SequenceRecorder.TabIcon"));
 
 		// Add cache tab.
 		FGlobalTabmanager::Get()->RegisterNomadTabSpawner(ImgMediaCacheTabName,
 			FOnSpawnTab::CreateStatic(&FImgMediaEditorModule::SpawnCacheTab))
 			.SetGroup(MediaBrowserGroup)
-			.SetDisplayName(LOCTEXT("ImgMediaCacheTabTitle", "Cache"))
-			.SetTooltipText(LOCTEXT("ImgMediaCacheTooltipText", "Open the cache tab."))
+			.SetDisplayName(LOCTEXT("ImgMediaGlobalCacheTabTitle", "Global Cache"))
+			.SetTooltipText(LOCTEXT("ImgMediaGlobalCacheTooltipText", "Open the global cache tab."))
 			.SetIcon(FSlateIcon(FAppStyle::GetAppStyleSetName(), "SequenceRecorder.TabIcon"));
 
 		// Add process images tab.
 		FGlobalTabmanager::Get()->RegisterNomadTabSpawner(ImgMediaProcessEXRTabName,
 			FOnSpawnTab::CreateStatic(&FImgMediaEditorModule::SpawnProcessEXRTab))
 			.SetGroup(MediaBrowserGroup)
-			.SetDisplayName(LOCTEXT("ImgMediaProcessImagesTabTitle", "Process EXR"))
-			.SetTooltipText(LOCTEXT("ImgMediaProcessImagesTooltipText", "Open the Process EXR tab."))
+			.SetDisplayName(LOCTEXT("ImgMediaProcessEXRTabTitle", "Process EXR"))
+			.SetTooltipText(LOCTEXT("ImgMediaProcessEXRTooltipText", "Open the Process EXR tab."))
 			.SetIcon(FSlateIcon(FAppStyle::GetAppStyleSetName(), "LevelEditor.Tabs.Viewports"));
 	}
 
