ImgMedia: Added process exr to context menu of ImgMediaSource.

#rnx

diff --git a/Engine/Plugins/Media/ImgMedia/Source/ImgMediaEditor/ImgMediaEditor.Build.cs b/Engine/Plugins/Media/ImgMedia/Source/ImgMediaEditor/ImgMediaEditor.Build.cs
--- a/Engine/Plugins/Media/ImgMedia/Source/ImgMediaEditor/ImgMediaEditor.Build.cs
+++ b/Engine/Plugins/Media/ImgMedia/Source/ImgMediaEditor/ImgMediaEditor.Build.cs
@@ -9,6 +9,7 @@ namespace UnrealBuildTool.Rules
 			PrivateDependencyModuleNames.AddRange(
 				new string[] {
 					"AssetTools",
+					"ContentBrowser",
 					"Core",
 					"CoreUObject",
 					"DesktopWidgets",				
@@ -20,6 +21,7 @@ namespace UnrealBuildTool.Rules
 					"RenderCore",
 					"Slate",
 					"SlateCore",
+					"ToolMenus",
 					"UnrealEd",
 					"WorkspaceMenuStructure",
 				});
diff --git a/Engine/Plugins/Media/ImgMedia/Source/ImgMediaEditor/Private/ImgMediaEditorModule.cpp b/Engine/Plugins/Media/ImgMedia/Source/ImgMediaEditor/Private/ImgMediaEditorModule.cpp
--- a/Engine/Plugins/Media/ImgMedia/Source/ImgMediaEditor/Private/ImgMediaEditorModule.cpp
+++ b/Engine/Plugins/Media/ImgMedia/Source/ImgMediaEditor/Private/ImgMediaEditorModule.cpp
@@ -5,12 +5,14 @@
 
 #include "AssetToolsModule.h"
 #include "AssetTools/ImgMediaSourceActions.h"
+#include "ContentBrowserMenuContexts.h"
 #include "Customizations/ImgMediaSourceCustomization.h"
 #include "IAssetTools.h"
 #include "ImgMediaEditorModule.h"
 #include "IImgMediaModule.h"
 #include "ImgMediaSource.h"
 #include "PropertyEditorModule.h"
+#include "ToolMenus.h"
 #include "UObject/NameTypes.h"
 #include "Widgets/Docking/SDockTab.h"
 #include "Widgets/SImgMediaBandwidth.h"
@@ -48,6 +50,7 @@ public:
 		RegisterCustomizations();
 		RegisterAssetTools();
 		RegisterTabSpawners();
+		ExtendContentMenu();
 
 		IImgMediaModule* ImgMediaModule = FModuleManager::LoadModulePtr<IImgMediaModule>("ImgMedia");
 		if (ImgMediaModule != nullptr)
@@ -185,6 +188,54 @@ protected:
 			];
 	}
 
+	void ExtendContentMenu()
+	{
+		if (UToolMenu* WorldAssetMenu = UToolMenus::Get()->ExtendMenu("ContentBrowser.AssetContextMenu.ImgMediaSource"))
+		{
+			FToolMenuSection& Section = WorldAssetMenu->AddDynamicSection("ImgMediaSource", FNewToolMenuDelegate::CreateLambda([this](UToolMenu* ToolMenu)
+			{
+				if (ToolMenu)
+				{
+					if (UContentBrowserAssetContextMenuContext* AssetMenuContext = ToolMenu->Context.FindContext<UContentBrowserAssetContextMenuContext>())
+					{
+						if (AssetMenuContext->SelectedObjects.Num() != 1)
+						{
+							return;
+						}
+
+						// Add menu entry.
+						if (UImgMediaSource* WorldAsset = Cast<UImgMediaSource>(AssetMenuContext->SelectedObjects[0].Get()))
+						{
+							const TAttribute<FText> Label = LOCTEXT("ProcessEXR", "Process EXRs");
+							const TAttribute<FText> ToolTip = LOCTEXT("ProcessEXR_Tooltip", "Open a tab to allow adding tiles and mips to an EXR sequence.");
+							const FSlateIcon Icon = FSlateIcon(FAppStyle::GetAppStyleSetName(), "LevelEditor.Tabs.Viewports");
+							const FToolMenuExecuteAction UIAction =
+								FToolMenuExecuteAction::CreateStatic(&FImgMediaEditorModule::OpenProcesEXRTab);
+						
+							const FName SectionName = TEXT("ImgMedia");
+							FToolMenuSection* Section = ToolMenu->FindSection(SectionName);
+							if (!Section)
+							{
+								Section = &(ToolMenu->AddSection(SectionName, LOCTEXT("ImgMediaSectionLabel", "ImgMedia")));
+							}
+
+							Section->AddMenuEntry("ImgMedia_ProcessEXR", Label, ToolTip, Icon, UIAction);
+						}
+					}
+				}
+			}), FToolMenuInsert(NAME_None, EToolMenuInsertType::Default));
+		}
+	}
+
+	/**
+	 * Call this to open the tab to process EXRs.
+	 */
+	static void OpenProcesEXRTab(const FToolMenuContext& MenuContext)
+	{
+		FTabId TabID(ImgMediaProcessEXRTabName);
+		FGlobalTabmanager::Get()->TryInvokeTab(TabID);
+	}
+
 	void OnImgMediaPlayerCreated(const TSharedPtr<FImgMediaPlayer>& Player)
 	{
 		// Try and replace an expired player.
