ImgMedia: Process EXR now relies on FImage for loading and conversions.

This fixes the immediate issue with jpgs, but generally corrects the conversion and linearization to float16 from various source image sequence formats.

diff --git a/Engine/Plugins/Media/ImgMedia/Source/ImgMediaEditor/ImgMediaEditor.Build.cs b/Engine/Plugins/Media/ImgMedia/Source/ImgMediaEditor/ImgMediaEditor.Build.cs
--- a/Engine/Plugins/Media/ImgMedia/Source/ImgMediaEditor/ImgMediaEditor.Build.cs
+++ b/Engine/Plugins/Media/ImgMedia/Source/ImgMediaEditor/ImgMediaEditor.Build.cs
@@ -14,6 +14,7 @@ namespace UnrealBuildTool.Rules
 					"CoreUObject",
 					"DesktopWidgets",				
 					"EditorFramework",
+					"ImageCore",
 					"ImageWrapper",
 					"ImgMedia",
 					"InputCore",
diff --git a/Engine/Plugins/Media/ImgMedia/Source/ImgMediaEditor/Private/Widgets/SImgMediaProcessEXR.cpp b/Engine/Plugins/Media/ImgMedia/Source/ImgMediaEditor/Private/Widgets/SImgMediaProcessEXR.cpp
--- a/Engine/Plugins/Media/ImgMedia/Source/ImgMediaEditor/Private/Widgets/SImgMediaProcessEXR.cpp
+++ b/Engine/Plugins/Media/ImgMedia/Source/ImgMediaEditor/Private/Widgets/SImgMediaProcessEXR.cpp
@@ -233,9 +233,7 @@ void SImgMediaProcessEXR::ProcessAllImages()
 			UE_LOG(LogImgMediaEditor, Error, TEXT("Invalid file format %s"), *Ext);
 		}
 		else
-		{
-			IImageWrapperModule& ImageWrapperModule = FModuleManager::LoadModuleChecked<IImageWrapperModule>("ImageWrapper");
-			
+		{			
 			// ImageWrapper is always returning an alpha channel for RGB, so check if we really have one.
 			bool bHasAlphaChannel = HasAlphaChannel(Ext, FPaths::Combine(SequencePath, FoundFiles[0]));
 
@@ -272,32 +270,37 @@ void SImgMediaProcessEXR::ProcessAllImages()
 				});
 				NumDone++;
 
-				TSharedPtr<IImageWrapper> ImageWrapper = ImageWrapperModule.CreateImageWrapper(ImageFormat);
 				Async(EAsyncExecution::Thread, [this, SequencePath, FileName, OutPath,
-					Ext, &NumActive,
-					ImageWrapper, InTileWidth, InTileHeight,
+					Ext, &NumActive, InTileWidth, InTileHeight,
 					TileBorder, bEnableMips, bHasAlphaChannel]() mutable
 				{
 					FString FullFileName = FPaths::Combine(SequencePath, FileName);
+					FImage Image;
 
 					// Load image into buffer.
-					TArray64<uint8> InputBuffer;
-					if (!FFileHelper::LoadFileToArray(InputBuffer, *FullFileName))
 					{
-						UE_LOG(LogImgMediaEditor, Error, TEXT("Failed to load %s"), *FullFileName);
-						NumActive--;
-						return;
+						TRACE_CPUPROFILER_EVENT_SCOPE(SImgMediaProcessEXR::ProcessImageCustom:LoadImage);
+						if (!FImageUtils::LoadImage(*FullFileName, Image))
+						{
+							UE_LOG(LogImgMediaEditor, Error, TEXT("Failed to load %s"), *FullFileName);
+							NumActive--;
+							return;
+						}
 					}
-					if (!ImageWrapper.IsValid() || !ImageWrapper->SetCompressed(InputBuffer.GetData(), InputBuffer.Num()))
+
+					// Convert to linear float16 if not already.
 					{
-						UE_LOG(LogImgMediaEditor, Error, TEXT("Failed to create image wrapper for %s"), *FullFileName);
-						NumActive--;
-						return;
+						TRACE_CPUPROFILER_EVENT_SCOPE(SImgMediaProcessEXR::ProcessImageCustom:ChangeFormat);
+						Image.ChangeFormat(ERawImageFormat::RGBA16F, EGammaSpace::Linear);
 					}
 
-					// Import this image.
 					FString Name = FPaths::Combine(OutPath, FileName);
-					ProcessImageCustom(ImageWrapper, InTileWidth, InTileHeight, TileBorder,
+					if (Ext != TEXT("exr"))
+					{
+						Name = FPaths::ChangeExtension(Name, TEXT("exr"));
+					}
+
+					ProcessImageCustom(Image, InTileWidth, InTileHeight, TileBorder,
 						bEnableMips, bHasAlphaChannel, Name, false /*bUseCustomFormat*/);
 
 					NumActive--;
@@ -348,7 +351,7 @@ bool SImgMediaProcessEXR::HasAlphaChannel(const FString& Ext, const FString& Fil
 	return bHasAlpha;
 }
 
-void SImgMediaProcessEXR::ProcessImageCustom(TSharedPtr<IImageWrapper>& InImageWrapper,
+void SImgMediaProcessEXR::ProcessImageCustom(const FImage& InImage,
 	int32 InTileWidth, int32 InTileHeight, int32 InTileBorder, bool bInEnableMips,
 	bool bHasAlphaChannel, const FString& InName,
 	bool bIsCustomFormat)
@@ -356,16 +359,11 @@ void SImgMediaProcessEXR::ProcessImageCustom(TSharedPtr<IImageWrapper>& InImageW
 #if IMGMEDIAEDITOR_EXR_SUPPORTED_PLATFORM
 	TRACE_CPUPROFILER_EVENT_SCOPE(SImgMediaProcessEXR::ProcessImageCustom);
 	// Get image data.
-	ERGBFormat Format = InImageWrapper->GetFormat();
-	int32 Width = InImageWrapper->GetWidth();
-	int32 Height = InImageWrapper->GetHeight();
-	int32 BitDepth = InImageWrapper->GetBitDepth();
-	TArray64<uint8> RawData;
-	{
-		TRACE_CPUPROFILER_EVENT_SCOPE(SImgMediaProcessEXR::ProcessImageCustom:GetRaw);
-		InImageWrapper->GetRaw(Format, BitDepth, RawData);
-	}
-	ProcessImageCustomRawData(RawData, Width, Height, BitDepth,
+	int32 Width = InImage.GetWidth();
+	int32 Height = InImage.GetHeight();
+	TArray64<uint8> RawData = InImage.RawData;
+
+	ProcessImageCustomRawData(RawData, Width, Height,
 		InTileWidth, InTileHeight, InTileBorder, bInEnableMips,
 		bHasAlphaChannel, InName, bIsCustomFormat);
 #else // IMGMEDIAEDITOR_EXR_SUPPORTED_PLATFORM
@@ -374,7 +372,7 @@ void SImgMediaProcessEXR::ProcessImageCustom(TSharedPtr<IImageWrapper>& InImageW
 }
 
 void SImgMediaProcessEXR::ProcessImageCustomRawData(TArray64<uint8>& RawData,
-	int32 Width, int32 Height, int32 BitDepth,
+	int32 Width, int32 Height,
 	int32 InTileWidth, int32 InTileHeight, int32 InTileBorder, bool bInEnableMips,
 	bool bHasAlphaChannel, const FString& InName,
 	bool bIsCustomFormat)
@@ -388,23 +386,10 @@ void SImgMediaProcessEXR::ProcessImageCustomRawData(TArray64<uint8>& RawData,
 	int32 TileWidth = InTileWidth;
 	int32 TileHeight = InTileHeight;
 	int32 BytesPerPixel = RawData.Num() / (Width * Height);
-	int32 BytesPerPixelPerChannel = BitDepth / 8;
+	int32 BytesPerPixelPerChannel = 16 / 8;
 	int32 NumChannels = BytesPerPixel / BytesPerPixelPerChannel;
 	int32 DestNumChannels = NumChannels;
 
-	// We use 16 bits only, so convert if we have 32.
-	if (BitDepth == 32)
-	{
-		ConvertTo16Bit(RawData);
-		BitDepth = 16;
-		BytesPerPixel /= 2;
-		BytesPerPixelPerChannel /= 2;
-	}
-	else if (BitDepth != 16)
-	{
-		UE_LOG(LogImgMediaEditor, Error, TEXT("Unsupported bit depth %d"), BitDepth);
-	}
-
 	// ImageWrapper always returns an alpha channel, so make sure we really have one.
 	if ((DestNumChannels == 4) && (bHasAlphaChannel == false))
 	{
@@ -672,27 +657,6 @@ void SImgMediaProcessEXR::ProcessImageCustomRawData(TArray64<uint8>& RawData,
 #endif // IMGMEDIAEDITOR_EXR_SUPPORTED_PLATFORM
 }
 
-
-void SImgMediaProcessEXR::ConvertTo16Bit(TArray64<uint8>& Buffer)
-{
-	TRACE_CPUPROFILER_EVENT_SCOPE(SImgMediaProcessEXR::ConvertTo16Bit);
-
-	int32 BytesPerPixelPerChannel = 4;
-	int64 BufferSize = Buffer.Num() / BytesPerPixelPerChannel;
-	FFloat32* InBufferPtr = (FFloat32*)(Buffer.GetData());
-	FFloat16* OutBufferPtr = (FFloat16*)(Buffer.GetData());
-
-	// Loop through the buffer.
-	for (int64 Index = 0; Index < BufferSize; ++Index)
-	{
-		// Convert the data in place.
-		OutBufferPtr[Index] = InBufferPtr[Index].FloatValue;
-	}
-
-	// Don't bother shrinking as its just a waste and extra work.
-	Buffer.SetNum(BufferSize / 2, false);
-}
-
 void SImgMediaProcessEXR::RemoveAlphaChannel(TArray64<uint8>& Buffer)
 {
 	TRACE_CPUPROFILER_EVENT_SCOPE(SImgMediaProcessEXR::RemoveAlphaChannel);
@@ -924,7 +888,6 @@ void SImgMediaProcessEXR::HandleProcessing()
 					{
 						int32 Width = RenderTarget->GetSurfaceWidth();
 						int32 Height = RenderTarget->GetSurfaceHeight();
-						int32 BitDepth = 16;
 						int32 InTileWidth = Options->bEnableTiling ? Options->TileSizeX : 0;
 						int32 InTileHeight = Options->bEnableTiling ? Options->TileSizeY : 0;
 						int32 TileBorder = 0; // Note: virtual texture support is shelved for now.
@@ -934,11 +897,11 @@ void SImgMediaProcessEXR::HandleProcessing()
 						FString FileName = FString::Printf(TEXT("image%05d.exr"), CurrentFrameIndex);
 						FString Name = FPaths::Combine(OutPath, FileName);
 
-						Async(EAsyncExecution::Thread, [this, RawData = MoveTemp(RawData), Width, Height, BitDepth,
+						Async(EAsyncExecution::Thread, [this, RawData = MoveTemp(RawData), Width, Height,
 							InTileWidth, InTileHeight, TileBorder, bEnableMips,
 							bHasAlphaChannel, Name]() mutable
 						{
-							ProcessImageCustomRawData(RawData, Width, Height, BitDepth,
+							ProcessImageCustomRawData(RawData, Width, Height,
 								InTileWidth, InTileHeight, TileBorder, bEnableMips,
 								bHasAlphaChannel, Name, false /*bUseCustomFormat*/);
 						});
diff --git a/Engine/Plugins/Media/ImgMedia/Source/ImgMediaEditor/Private/Widgets/SImgMediaProcessEXR.h b/Engine/Plugins/Media/ImgMedia/Source/ImgMediaEditor/Private/Widgets/SImgMediaProcessEXR.h
--- a/Engine/Plugins/Media/ImgMedia/Source/ImgMediaEditor/Private/Widgets/SImgMediaProcessEXR.h
+++ b/Engine/Plugins/Media/ImgMedia/Source/ImgMediaEditor/Private/Widgets/SImgMediaProcessEXR.h
@@ -9,7 +9,7 @@
 #include "Widgets/SCompoundWidget.h"
 
 class IDetailsView;
-class IImageWrapper;
+struct FImage;
 class SButton;
 class SNotificationItem;
 class UMediaPlayer;
@@ -64,7 +64,7 @@ private:
 	 * Tiles and mips may be generated.
 	 * This does NOT run on the game thread.
 	 *
-	 * @param InImageWrapper	ImageWrapper to read/write the image.
+	 * @param InImage			FImage of the source image.
 	 * @param InTileWidth		Desired width of tiles.
 	 * @param InTileHeight		Desired height of tiles.
 	 * @param InTileBorder		Number of pixels to duplicate along a tile edge.
@@ -73,7 +73,7 @@ private:
 	 * @param InName			Full path and name of file to write.
 	 * @param bIsCustomFormat	True to output in our custom format.
 	 */
-	void ProcessImageCustom(TSharedPtr<IImageWrapper>& InImageWrapper,
+	void ProcessImageCustom(const FImage& InImage,
 		int32 InTileWidth, int32 InTileHeight, int32 InTileBorder, 
 		bool bInEnableMips, bool bHasAlphaChannel, const FString& InName,
 		bool bIsCustomFormat);
@@ -86,7 +86,6 @@ private:
 	 * @param RawData			Image pixel data (RGBA FFloat16).
 	 * @param Width				Width of image.
 	 * @param Height			Height of image.
-	 * @param BitDepth			Number of bits per pixel.
 	 * @param InTileWidth		Desired width of tiles.
 	 * @param InTileHeight		Desired height of tiles.
 	 * @param InTileBorder		Number of pixels to duplicate along a tile edge.
@@ -96,17 +95,9 @@ private:
 	 * @param bIsCustomFormat	True to output in our custom format.
 	 */
 	void ProcessImageCustomRawData(TArray64<uint8>& RawData,
-		int32 Width, int32 Height, int32 BitDepth,
+		int32 Width, int32 Height,
 		int32 InTileWidth, int32 InTileHeight, int32 InTileBorder, bool bInEnableMips,
 		bool bHasAlphaChannel, const FString& InName, bool bIsCustomFormat);
-	
-	/**
-	 * Converts a buffer from 32 bit floating point to 16 bit floating point.
-	 * This will be done in place.
-	 * 
-	 * @param Buffer			Buffer to convert.
-	 */
-	void ConvertTo16Bit(TArray64<uint8>& Buffer);
 
 	/**
 	 * Removes the alpha channel from a buffer.
