ImgMedia: Opening process exr via a media source will set the input path to the media source.

#rnx

diff --git a/Engine/Plugins/Media/ImgMedia/Source/ImgMediaEditor/Private/ImgMediaEditorModule.cpp b/Engine/Plugins/Media/ImgMedia/Source/ImgMediaEditor/Private/ImgMediaEditorModule.cpp
--- a/Engine/Plugins/Media/ImgMedia/Source/ImgMediaEditor/Private/ImgMediaEditorModule.cpp
+++ b/Engine/Plugins/Media/ImgMedia/Source/ImgMediaEditor/Private/ImgMediaEditorModule.cpp
@@ -181,11 +181,22 @@ protected:
 
 	static TSharedRef<SDockTab> SpawnProcessEXRTab(const FSpawnTabArgs& SpawnTabArgs)
 	{
-		return SNew(SDockTab)
+		// Create tab.
+		TSharedPtr< SImgMediaProcessEXR> ProcessEXR;
+		TSharedRef <SDockTab> Tab = SNew(SDockTab)
 			.TabRole(ETabRole::NomadTab)
 			[
-				SNew(SImgMediaProcessEXR)
+				SAssignNew(ProcessEXR, SImgMediaProcessEXR)
 			];
+
+		// Override input path if desired.
+		if (ProcessEXRInputPath.IsEmpty() == false)
+		{
+			ProcessEXR->SetInputPath(ProcessEXRInputPath);
+			ProcessEXRInputPath.Empty();
+		}
+
+		return Tab;
 	}
 
 	void ExtendContentMenu()
@@ -206,11 +217,25 @@ protected:
 						// Add menu entry.
 						if (UImgMediaSource* WorldAsset = Cast<UImgMediaSource>(AssetMenuContext->SelectedObjects[0].Get()))
 						{
+							// Get path from media source.
+							// Make sure it ends with / so it opens the directory.
+							FString InputPath = WorldAsset->GetSequencePath();
+							if (InputPath.EndsWith(TEXT("/")) == false)
+							{
+								InputPath += TEXT("/");
+							}
+
 							const TAttribute<FText> Label = LOCTEXT("ProcessEXR", "Process EXRs");
 							const TAttribute<FText> ToolTip = LOCTEXT("ProcessEXR_Tooltip", "Open a tab to allow adding tiles and mips to an EXR sequence.");
-							const FSlateIcon Icon = FSlateIcon(FAppStyle::GetAppStyleSetName(), "LevelEditor.Tabs.Viewports");
+							const FSlateIcon Icon = FSlateIcon(FAppStyle::GetAppStyleSetName(),
+								"LevelEditor.Tabs.Viewports");
 							const FToolMenuExecuteAction UIAction =
-								FToolMenuExecuteAction::CreateStatic(&FImgMediaEditorModule::OpenProcesEXRTab);
+								FToolMenuExecuteAction::CreateLambda([InputPath]
+								(const FToolMenuContext& InContext)
+							{
+								ProcessEXRInputPath = InputPath;
+								OpenProcesEXRTab();
+							});
 						
 							const FName SectionName = TEXT("ImgMedia");
 							FToolMenuSection* Section = ToolMenu->FindSection(SectionName);
@@ -230,7 +255,7 @@ protected:
 	/**
 	 * Call this to open the tab to process EXRs.
 	 */
-	static void OpenProcesEXRTab(const FToolMenuContext& MenuContext)
+	static void OpenProcesEXRTab()
 	{
 		FTabId TabID(ImgMediaProcessEXRTabName);
 		FGlobalTabmanager::Get()->TryInvokeTab(TabID);
@@ -268,6 +293,9 @@ private:
 	/** Names for tabs. */
 	static FLazyName ImgMediaProcessEXRTabName;
 
+	/** This will be passed to ProceessEXR and then cleared. */
+	static FString ProcessEXRInputPath;
+
 	/** The collection of registered asset type actions. */
 	TArray<TSharedRef<IAssetTypeActions>> RegisteredAssetTypeActions;
 
@@ -276,6 +304,7 @@ private:
 };
 
 FLazyName FImgMediaEditorModule::ImgMediaProcessEXRTabName(TEXT("ImgMediaProcessEXR"));
+FString FImgMediaEditorModule::ProcessEXRInputPath;
 
 IMPLEMENT_MODULE(FImgMediaEditorModule, ImgMediaEditor);
 
diff --git a/Engine/Plugins/Media/ImgMedia/Source/ImgMediaEditor/Private/Widgets/SImgMediaProcessEXR.cpp b/Engine/Plugins/Media/ImgMedia/Source/ImgMediaEditor/Private/Widgets/SImgMediaProcessEXR.cpp
--- a/Engine/Plugins/Media/ImgMedia/Source/ImgMediaEditor/Private/Widgets/SImgMediaProcessEXR.cpp
+++ b/Engine/Plugins/Media/ImgMedia/Source/ImgMediaEditor/Private/Widgets/SImgMediaProcessEXR.cpp
@@ -111,6 +111,14 @@ void SImgMediaProcessEXR::Construct(const FArguments& InArgs)
 }
 END_SLATE_FUNCTION_BUILD_OPTIMIZATION
 
+void SImgMediaProcessEXR::SetInputPath(const FString& Path)
+{
+	if (Options.IsValid())
+	{
+		Options->InputPath.FilePath = Path;
+	}
+}
+
 void SImgMediaProcessEXR::Tick(const FGeometry& AllottedGeometry, const double InCurrentTime, const float InDeltaTime)
 {
 	SCompoundWidget::Tick(AllottedGeometry, InCurrentTime, InDeltaTime);
diff --git a/Engine/Plugins/Media/ImgMedia/Source/ImgMediaEditor/Private/Widgets/SImgMediaProcessEXR.h b/Engine/Plugins/Media/ImgMedia/Source/ImgMediaEditor/Private/Widgets/SImgMediaProcessEXR.h
--- a/Engine/Plugins/Media/ImgMedia/Source/ImgMediaEditor/Private/Widgets/SImgMediaProcessEXR.h
+++ b/Engine/Plugins/Media/ImgMedia/Source/ImgMediaEditor/Private/Widgets/SImgMediaProcessEXR.h
@@ -30,6 +30,11 @@ public:
 
 	void Construct(const FArguments& InArgs);
 
+	/**
+	 * Call this to set what the input path should be.
+	 */
+	void SetInputPath(const FString& Path);
+
 	/** SWidget interface */
 	virtual void Tick(const FGeometry& AllottedGeometry, const double InCurrentTime, const float InDeltaTime);
 
