ImgMedia: Loader uses the same sequence index as the frame request to prevent any mismatches from differing logic between the loader and MediaFramework.

#rnx

diff --git a/Engine/Plugins/Media/ImgMedia/Source/ImgMedia/Private/Loader/ImgMediaLoader.cpp b/Engine/Plugins/Media/ImgMedia/Source/ImgMedia/Private/Loader/ImgMediaLoader.cpp
--- a/Engine/Plugins/Media/ImgMedia/Source/ImgMedia/Private/Loader/ImgMediaLoader.cpp
+++ b/Engine/Plugins/Media/ImgMedia/Source/ImgMedia/Private/Loader/ImgMediaLoader.cpp
@@ -523,47 +523,29 @@ IMediaSamples::EFetchBestSampleResult FImgMediaLoader::FetchBestVideoSampleForTi
 				RetryCount = 0;
 				double Duration = Frame->Get()->Info.FrameRate.AsInterval();
 
-				// Yes. First time (after flush)?
-				int32 NewSequenceIndex = QueuedSampleFetch.CurrentSequenceIndex;
-				if (QueuedSampleFetch.LastFrameIndex != INDEX_NONE)
+				// Yes.
+				int32 NewSequenceIndex = 0;
+				
+				// Make sure this sample has the sequence index that matches the time range.
+				// If the time range only has one sequence index then just use that.
+				if (TimeRange.GetLowerBoundValue().SequenceIndex == TimeRange.GetUpperBoundValue().SequenceIndex)
 				{
-					// No. Check if we looped and need to start a new sequence...
-					if (PlayRate >= 0.0f && QueuedSampleFetch.LastFrameIndex > MaxIdx)
-					{
-						++NewSequenceIndex;
-					}
-					else if (PlayRate < 0.0f && QueuedSampleFetch.LastFrameIndex < MaxIdx)
-					{
-						--NewSequenceIndex;
-					}
-					else if (GetNumImages() == 1)
-					{
-						NewSequenceIndex = TimeRange.GetLowerBoundValue().SequenceIndex;
-					}
+					NewSequenceIndex = TimeRange.GetLowerBoundValue().SequenceIndex;
 				}
 				else
 				{
-					// Make sure this sample has the sequence index that matches the time range.
-					// If the time range only has one sequence index then just use that.
-					if (TimeRange.GetLowerBoundValue().SequenceIndex == TimeRange.GetUpperBoundValue().SequenceIndex)
-					{
-						NewSequenceIndex = TimeRange.GetLowerBoundValue().SequenceIndex;
-					}
-					else
-					{
-						// Try the lower bound sequence index.
-						NewSequenceIndex = TimeRange.GetLowerBoundValue().SequenceIndex;
-						FMediaTimeStamp SampleTime = FMediaTimeStamp(FrameNumberToTime(MaxIdx), NewSequenceIndex);
-						TRange<FMediaTimeStamp> SampleTimeRange(
-							SampleTime,
-							SampleTime + FTimespan::FromSeconds(Duration));
+					// Try the lower bound sequence index.
+					NewSequenceIndex = TimeRange.GetLowerBoundValue().SequenceIndex;
+					FMediaTimeStamp SampleTime = FMediaTimeStamp(FrameNumberToTime(MaxIdx), NewSequenceIndex);
+					TRange<FMediaTimeStamp> SampleTimeRange(
+						SampleTime,
+						SampleTime + FTimespan::FromSeconds(Duration));
 						
-						// Does our sample time overlap the time range?
-						if (TimeRange.Overlaps(SampleTimeRange) == false)
-						{
-							// No. Use the upper bound sequence index.
-							NewSequenceIndex = TimeRange.GetUpperBoundValue().SequenceIndex;
-						}
+					// Does our sample time overlap the time range?
+					if (TimeRange.Overlaps(SampleTimeRange) == false)
+					{
+						// No. Use the upper bound sequence index.
+						NewSequenceIndex = TimeRange.GetUpperBoundValue().SequenceIndex;
 					}
 				}
 				
