ImgMedia: Improvements to the loader. We avoid queuing empty work, and sort the tile selection map by mip level for consistency.

#rnx

diff --git a/Engine/Plugins/Media/ImgMedia/Source/ImgMedia/Private/Loader/ImgMediaLoader.cpp b/Engine/Plugins/Media/ImgMedia/Source/ImgMedia/Private/Loader/ImgMediaLoader.cpp
--- a/Engine/Plugins/Media/ImgMedia/Source/ImgMedia/Private/Loader/ImgMediaLoader.cpp
+++ b/Engine/Plugins/Media/ImgMedia/Source/ImgMedia/Private/Loader/ImgMediaLoader.cpp
@@ -681,12 +681,22 @@ IQueuedWork* FImgMediaLoader::GetWork()
 {
 	FScopeLock Lock(&CriticalSection);
 
-	if (PendingFrameNumbers.Num() == 0)
+	TMap<int32, FImgMediaTileSelection> DesiredMipsAndTiles;
+	int32 FrameNumber = INDEX_NONE;
+
+	// Find a visible pending frame.
+	while (!PendingFrameNumbers.IsEmpty() && DesiredMipsAndTiles.IsEmpty())
+	{
+		FrameNumber = PendingFrameNumbers.Pop(false);
+		GetDesiredMipTiles(FrameNumber, DesiredMipsAndTiles);
+	}
+
+	// No selection is visible (or pending frames was empty), so we don't queue any work.
+	if (DesiredMipsAndTiles.IsEmpty())
 	{
 		return nullptr;
 	}
 
-	int32 FrameNumber = PendingFrameNumbers.Pop(false);
 	FImgMediaLoaderWork* Work = (WorkPool.Num() > 0) ? WorkPool.Pop() : new FImgMediaLoaderWork(AsShared(), Reader.ToSharedRef());
 	
 	// Get the existing frame so we can add the mip level to it.
@@ -706,9 +716,6 @@ IQueuedWork* FImgMediaLoader::GetWork()
 	{
 		ExistingFrame = *ExistingFramePtr;
 	}
-
-	TMap<int32, FImgMediaTileSelection> DesiredMipsAndTiles;
-	GetDesiredMipTiles(FrameNumber, DesiredMipsAndTiles);
 	
 	// Set up work.
 	Work->Initialize(FrameNumber, DesiredMipsAndTiles, ExistingFrame);
@@ -1367,6 +1374,8 @@ void FImgMediaLoader::AddFrameToCache(int32 FrameNumber, const TSharedPtr<FImgMe
 
 void FImgMediaLoader::GetDesiredMipTiles(int32 FrameIndex, TMap<int32, FImgMediaTileSelection>& OutMipsAndTiles)
 {
+	// note: While currently unused, FrameIndex could become useful for recorded frames.
+
 	if(MipMapInfo.IsValid() && MipMapInfo->HasObjects())
 	{
 		OutMipsAndTiles = MipMapInfo->GetVisibleTiles();
diff --git a/Engine/Plugins/Media/ImgMedia/Source/ImgMedia/Private/Loader/ImgMediaLoaderWork.cpp b/Engine/Plugins/Media/ImgMedia/Source/ImgMedia/Private/Loader/ImgMediaLoaderWork.cpp
--- a/Engine/Plugins/Media/ImgMedia/Source/ImgMedia/Private/Loader/ImgMediaLoaderWork.cpp
+++ b/Engine/Plugins/Media/ImgMedia/Source/ImgMedia/Private/Loader/ImgMediaLoaderWork.cpp
@@ -39,6 +39,9 @@ void FImgMediaLoaderWork::Initialize(int32 InFrameNumber,
 	FrameNumber = InFrameNumber;
 	MipTiles = MoveTemp(InMipTiles);
 	ExistingFrame = InExistingFrame;
+
+	// Ensure that we start from mip0 when iterating the map for consistency
+	MipTiles.KeySort(TLess<int32>());
 }
 
 
