ImgMedia: Fix GenericImgMediaReader for mips in separate folders support.

Instead of relying on the currently read data for buffer allocation (which could be any mip level), we use the loader-provider dimensions. The map iteration was also incorrectly assuming in-order traversal.

#rnx

diff --git a/Engine/Plugins/Media/ImgMedia/Source/ImgMedia/Private/Readers/GenericImgMediaReader.cpp b/Engine/Plugins/Media/ImgMedia/Source/ImgMedia/Private/Readers/GenericImgMediaReader.cpp
--- a/Engine/Plugins/Media/ImgMedia/Source/ImgMedia/Private/Readers/GenericImgMediaReader.cpp
+++ b/Engine/Plugins/Media/ImgMedia/Source/ImgMedia/Private/Readers/GenericImgMediaReader.cpp
@@ -10,10 +10,7 @@
 #include "Misc/Paths.h"
 
 
-/* Local helpers
- *****************************************************************************/
-
-TSharedPtr<IImageWrapper> LoadImage(const FString& ImagePath, IImageWrapperModule& ImageWrapperModule, TArray64<uint8>& OutBuffer, FImgMediaFrameInfo& OutInfo)
+TSharedPtr<IImageWrapper> FGenericImgMediaReader::LoadFrameImage(const FString& ImagePath, TArray64<uint8>& OutBuffer, FImgMediaFrameInfo& OutInfo, bool bUseLoaderFirstFrame)
 {
 	// load image into buffer
 	if (!FFileHelper::LoadFileToArray(OutBuffer, *ImagePath))
@@ -57,25 +54,42 @@ TSharedPtr<IImageWrapper> LoadImage(const FString& ImagePath, IImageWrapperModul
 		return nullptr;
 	}
 
+	if (bUseLoaderFirstFrame && LoaderPtr.IsValid())
+	{
+		const TSharedPtr<FImgMediaLoader, ESPMode::ThreadSafe> Loader = LoaderPtr.Pin();
+
+		OutInfo.Dim = Loader->GetSequenceDim();
+		OutInfo.NumMipLevels = Loader->GetNumMipLevels();
+
+		const SIZE_T SizeMip0 = (SIZE_T)OutInfo.Dim.X * OutInfo.Dim.Y * 4;
+		OutInfo.UncompressedSize = SizeMip0;
+		for (int32 Level = 1; Level < OutInfo.NumMipLevels; Level++)
+		{
+			OutInfo.UncompressedSize += SizeMip0 >> (2 * Level);
+		}
+	}
+	else
+	{
+		OutInfo.Dim.X = ImageWrapper->GetWidth();
+		OutInfo.Dim.Y = ImageWrapper->GetHeight();
+		OutInfo.UncompressedSize = (SIZE_T)OutInfo.Dim.X * OutInfo.Dim.Y * 4;
+		OutInfo.NumMipLevels = 1;
+	}
+
 	// get file info
 	const UImgMediaSettings* Settings = GetDefault<UImgMediaSettings>();
 	OutInfo.CompressionName = TEXT("");
-	OutInfo.Dim.X = ImageWrapper->GetWidth();
-	OutInfo.Dim.Y = ImageWrapper->GetHeight();
 	OutInfo.FrameRate = Settings->DefaultFrameRate;
 	OutInfo.Srgb = true;
-	OutInfo.UncompressedSize = (SIZE_T)OutInfo.Dim.X * OutInfo.Dim.Y * 4;
 	OutInfo.NumChannels = 4;
 	OutInfo.bHasTiles = false;
 	OutInfo.TileDimensions = OutInfo.Dim;
 	OutInfo.NumTiles = FIntPoint(1, 1);
 	OutInfo.TileBorder = 0;
-	OutInfo.NumMipLevels = 1;
 
 	return ImageWrapper;
 }
 
-
 /* FGenericImgMediaReader structors
  *****************************************************************************/
 
@@ -91,7 +105,7 @@ FGenericImgMediaReader::FGenericImgMediaReader(IImageWrapperModule& InImageWrapp
 bool FGenericImgMediaReader::GetFrameInfo(const FString& ImagePath, FImgMediaFrameInfo& OutInfo)
 {
 	TArray64<uint8> InputBuffer;
-	TSharedPtr<IImageWrapper> ImageWrapper = LoadImage(ImagePath, ImageWrapperModule, InputBuffer, OutInfo);
+	TSharedPtr<IImageWrapper> ImageWrapper = LoadFrameImage(ImagePath, InputBuffer, OutInfo, false);
 
 	return ImageWrapper.IsValid();
 }
@@ -99,75 +113,80 @@ bool FGenericImgMediaReader::GetFrameInfo(const FString& ImagePath, FImgMediaFra
 
 bool FGenericImgMediaReader::ReadFrame(int32 FrameId, const TMap<int32, FImgMediaTileSelection>& InMipTiles, TSharedPtr<FImgMediaFrame, ESPMode::ThreadSafe> OutFrame)
 {
-	TSharedPtr<FImgMediaLoader, ESPMode::ThreadSafe> Loader = LoaderPtr.Pin();
+	const TSharedPtr<FImgMediaLoader, ESPMode::ThreadSafe> Loader = LoaderPtr.Pin();
 	if (Loader.IsValid() == false)
 	{
 		return false;
 	}
 
-	int32 NumMipLevels = Loader->GetNumMipLevels();
-	FIntPoint Dim = Loader->GetSequenceDim();
-	
-	// Do we already have our buffer?
-	void* Buffer = OutFrame->Data.Get();
+	if (InMipTiles.IsEmpty())
+	{
+		return false;
+	}
+
+	SIZE_T BufferDataOffset = 0;
+	const int32 NumMipLevels = Loader->GetNumMipLevels();
+	const FIntPoint BaseLevelDim = Loader->GetSequenceDim();
+	FIntPoint CurrenDim = BaseLevelDim;
 
 	// Loop over all mips.
-	for (const TPair<int32, FImgMediaTileSelection>& TilesPerMip : InMipTiles)
+	for (int32 CurrentMipLevel = 0; CurrentMipLevel < NumMipLevels; ++CurrentMipLevel)
 	{
-		const int32 CurrentMipLevel = TilesPerMip.Key;
-		const FImgMediaTileSelection& CurrentSelection = TilesPerMip.Value;
-
-		// Do we want to read in this mip?
-		bool IsThisLevelPresent = OutFrame->MipTilesPresent.Contains(CurrentMipLevel);
-		bool ReadThisMip = (Buffer == nullptr) || (IsThisLevelPresent == false);
-		if (ReadThisMip)
+		if (InMipTiles.Contains(CurrentMipLevel))
 		{
-			// Load image.
-			const FString& ImagePath = Loader->GetImagePath(FrameId, CurrentMipLevel);
+			const FImgMediaTileSelection& CurrentSelection = InMipTiles[CurrentMipLevel];
 
-			TArray64<uint8> InputBuffer;
-			FImgMediaFrameInfo Info;
-			TSharedPtr<IImageWrapper> ImageWrapper = LoadImage(ImagePath, ImageWrapperModule, InputBuffer, Info);
-
-			if (!ImageWrapper.IsValid())
+			// Do we want to read in this mip?
+			bool ReadThisMip = !OutFrame->MipTilesPresent.Contains(CurrentMipLevel);
+			if (ReadThisMip)
 			{
-				UE_LOG(LogImgMedia, Warning, TEXT("FGenericImgMediaReader: Failed to load image %s"), *ImagePath);
-				return false;
-			}
+				// Load image.
+				const FString& ImagePath = Loader->GetImagePath(FrameId, CurrentMipLevel);
 
-			TArray64<uint8> RawData;
-			if (!ImageWrapper->GetRaw(ERGBFormat::BGRA, 8, RawData))
-			{
-				UE_LOG(LogImgMedia, Warning, TEXT("FGenericImgMediaReader: Failed to get image data for %s"), *ImagePath);
-				return false;
-			}
+				TArray64<uint8> InputBuffer;
+				FImgMediaFrameInfo Info;
+				TSharedPtr<IImageWrapper> ImageWrapper = LoadFrameImage(ImagePath, InputBuffer, Info, true);
 
-			const int64 RawNum = RawData.Num();
-			// Create buffer for data.
-			if (Buffer == nullptr)
-			{
-				int64 AllocSize = RawNum;
-				// Need more space for mips.
-				if (NumMipLevels > 1)
+				if (!ImageWrapper.IsValid())
 				{
-					AllocSize = (AllocSize * 4) / 3;
+					UE_LOG(LogImgMedia, Warning, TEXT("FGenericImgMediaReader: Failed to load image %s"), *ImagePath);
+					return false;
 				}
-				Buffer = FMemory::Malloc(AllocSize);
-				OutFrame->Info = Info;
-				OutFrame->Data = MakeShareable(Buffer, [](void* ObjectToDelete) { FMemory::Free(ObjectToDelete); });
-				OutFrame->Format = EMediaTextureSampleFormat::CharBGRA;
-				OutFrame->Stride = OutFrame->Info.Dim.X * 4;
-			}
 
-			// Copy data to our buffer
-			FMemory::Memcpy(Buffer, RawData.GetData(), RawNum);
-			OutFrame->MipTilesPresent.Emplace(CurrentMipLevel, CurrentSelection);
-			OutFrame->NumTilesRead++;
+				TArray64<uint8> RawData;
+				if (!ImageWrapper->GetRaw(ERGBFormat::BGRA, 8, RawData))
+				{
+					UE_LOG(LogImgMedia, Warning, TEXT("FGenericImgMediaReader: Failed to get image data for %s"), *ImagePath);
+					return false;
+				}
+
+				// Create the full buffer for data.
+				if (!OutFrame->Data.IsValid())
+				{
+					int64 AllocSize = (SIZE_T)BaseLevelDim.X * BaseLevelDim.Y * 4;
+					// Need more space for mips.
+					if (NumMipLevels > 1)
+					{
+						AllocSize = (AllocSize * 4) / 3;
+					}
+					void* Buffer = FMemory::Malloc(AllocSize, PLATFORM_CACHE_LINE_SIZE);
+					OutFrame->Info = Info;
+					OutFrame->Data = MakeShareable(Buffer, [](void* ObjectToDelete) { FMemory::Free(ObjectToDelete); });
+					OutFrame->MipTilesPresent.Reset();
+					OutFrame->Format = EMediaTextureSampleFormat::CharBGRA;
+					OutFrame->Stride = OutFrame->Info.Dim.X * 4;
+				}
+
+				// Copy data to our buffer with the right mip level offset
+				FMemory::Memcpy((void*)((uint8*)OutFrame->Data.Get() + BufferDataOffset), RawData.GetData(), RawData.Num());
+				OutFrame->MipTilesPresent.Emplace(CurrentMipLevel, CurrentSelection);
+				OutFrame->NumTilesRead++;
+			}
 		}
 
 		// Next level.
-		Buffer = (void*)((uint8*)Buffer + Dim.X * Dim.Y * 4);
-		Dim /= 2;
+		BufferDataOffset += (CurrenDim.X * CurrenDim.Y * 4);
+		CurrenDim /= 2;
 	}
 
 	return true;
diff --git a/Engine/Plugins/Media/ImgMedia/Source/ImgMedia/Private/Readers/GenericImgMediaReader.h b/Engine/Plugins/Media/ImgMedia/Source/ImgMedia/Private/Readers/GenericImgMediaReader.h
--- a/Engine/Plugins/Media/ImgMedia/Source/ImgMedia/Private/Readers/GenericImgMediaReader.h
+++ b/Engine/Plugins/Media/ImgMedia/Source/ImgMedia/Private/Readers/GenericImgMediaReader.h
@@ -2,10 +2,12 @@
 
 #pragma once
 
+#include "Containers/ContainersFwd.h"
 #include "IImgMediaReader.h"
 
 class FImgMediaLoader;
 class FRgbaInputFile;
+class IImageWrapper;
 class IImageWrapperModule;
 
 
@@ -37,6 +39,8 @@ public:
 
 private:
 
+	TSharedPtr<IImageWrapper> LoadFrameImage(const FString& ImagePath, TArray64<uint8>& OutBuffer, FImgMediaFrameInfo& OutInfo, bool bUseLoaderFirstFrame);
+
 	/** The image wrapper module. */
 	IImageWrapperModule& ImageWrapperModule;
 	/** Our parent loader. */
