Mediaframework: Can retrieve the number of tiles from the player.

#rnx

diff --git a/Engine/Plugins/Media/ImgMedia/Source/ImgMedia/Private/Player/ImgMediaPlayer.cpp b/Engine/Plugins/Media/ImgMedia/Source/ImgMedia/Private/Player/ImgMediaPlayer.cpp
--- a/Engine/Plugins/Media/ImgMedia/Source/ImgMedia/Private/Player/ImgMediaPlayer.cpp
+++ b/Engine/Plugins/Media/ImgMedia/Source/ImgMedia/Private/Player/ImgMediaPlayer.cpp
@@ -287,6 +287,18 @@ bool FImgMediaPlayer::Open(const TSharedRef<FArchive, ESPMode::ThreadSafe>& /*Ar
 	return false; // not supported
 }
 
+FIntPoint FImgMediaPlayer::GetTileNum() const
+{
+	FIntPoint TileNum = FIntPoint::ZeroValue;
+
+	if (Loader.IsValid())
+	{
+		TileNum.X = Loader->GetNumTilesX();
+		TileNum.Y = Loader->GetNumTilesY();
+	}
+	
+	return TileNum;
+}
 
 void FImgMediaPlayer::TickInput(FTimespan DeltaTime, FTimespan /*Timecode*/)
 {
diff --git a/Engine/Plugins/Media/ImgMedia/Source/ImgMedia/Private/Player/ImgMediaPlayer.h b/Engine/Plugins/Media/ImgMedia/Source/ImgMedia/Private/Player/ImgMediaPlayer.h
--- a/Engine/Plugins/Media/ImgMedia/Source/ImgMedia/Private/Player/ImgMediaPlayer.h
+++ b/Engine/Plugins/Media/ImgMedia/Source/ImgMedia/Private/Player/ImgMediaPlayer.h
@@ -62,6 +62,7 @@ public:
 	virtual IMediaView& GetView() override;
 	virtual bool Open(const FString& Url, const IMediaOptions* Options) override;
 	virtual bool Open(const TSharedRef<FArchive, ESPMode::ThreadSafe>& Archive, const FString& OriginalUrl, const IMediaOptions* Options) override;
+	virtual FIntPoint GetTileNum() const override;
 	virtual void TickInput(FTimespan DeltaTime, FTimespan Timecode) override;
 	virtual bool FlushOnSeekStarted() const;
 	virtual bool FlushOnSeekCompleted() const;
diff --git a/Engine/Source/Runtime/Media/Public/IMediaPlayer.h b/Engine/Source/Runtime/Media/Public/IMediaPlayer.h
--- a/Engine/Source/Runtime/Media/Public/IMediaPlayer.h
+++ b/Engine/Source/Runtime/Media/Public/IMediaPlayer.h
@@ -5,6 +5,8 @@
 #include "Containers/UnrealString.h"
 #include "Internationalization/Text.h"
 #include "Math/Interval.h"
+#include "Math/IntPoint.h"
+#include "Math/MathFwd.h"
 #include "Misc/Paths.h"
 #include "Misc/Timespan.h"
 
@@ -202,6 +204,14 @@ public:
 		return FText::FromString(FPaths::GetBaseFilename(Url));
 	}
 
+	/**
+	 * Get the number of tiles in each frame.
+	 */
+	virtual FIntPoint GetTileNum() const
+	{
+		return FIntPoint::ZeroValue;
+	}
+
 	/**
 	 * Set the player's globally unique identifier.
 	 *
diff --git a/Engine/Source/Runtime/MediaAssets/Private/Assets/MediaPlayer.cpp b/Engine/Source/Runtime/MediaAssets/Private/Assets/MediaPlayer.cpp
--- a/Engine/Source/Runtime/MediaAssets/Private/Assets/MediaPlayer.cpp
+++ b/Engine/Source/Runtime/MediaAssets/Private/Assets/MediaPlayer.cpp
@@ -230,6 +230,10 @@ void UMediaPlayer::GetSupportedRates(TArray<FFloatRange>& OutRates, bool Unthinn
 	Rates.GetRanges((TArray<TRange<float>>&)OutRates);
 }
 
+FIntPoint UMediaPlayer::GetTileNum() const
+{
+	return PlayerFacade->GetTileNum();
+}
 
 FTimespan UMediaPlayer::GetTime() const
 {
diff --git a/Engine/Source/Runtime/MediaAssets/Public/MediaPlayer.h b/Engine/Source/Runtime/MediaAssets/Public/MediaPlayer.h
--- a/Engine/Source/Runtime/MediaAssets/Public/MediaPlayer.h
+++ b/Engine/Source/Runtime/MediaAssets/Public/MediaPlayer.h
@@ -297,6 +297,14 @@ public:
 	UFUNCTION(BlueprintCallable, Category="Media|MediaPlayer")
 	void GetSupportedRates(TArray<FFloatRange>& OutRates, bool Unthinned) const;
 
+	/**
+	 * Get the number of tiles in each frame.
+	 * 
+	 * @return Number of tiles in each dimension.
+	 */
+	UFUNCTION(BlueprintCallable, Category = "Media|MediaPlayer")
+	virtual FIntPoint GetTileNum() const;
+
 	/**
 	 * Get the media's current playback time.
 	 *
diff --git a/Engine/Source/Runtime/MediaUtils/Private/MediaPlayerFacade.cpp b/Engine/Source/Runtime/MediaUtils/Private/MediaPlayerFacade.cpp
--- a/Engine/Source/Runtime/MediaUtils/Private/MediaPlayerFacade.cpp
+++ b/Engine/Source/Runtime/MediaUtils/Private/MediaPlayerFacade.cpp
@@ -423,6 +423,18 @@ bool FMediaPlayerFacade::HaveAudioPlayback() const
 	return PrimaryAudioSink.IsValid() && GetSelectedTrack(EMediaTrackType::Audio) != INDEX_NONE;
 }
 
+FIntPoint FMediaPlayerFacade::GetTileNum() const
+{
+	FIntPoint TileNum = FIntPoint::ZeroValue;
+
+	TSharedPtr<IMediaPlayer, ESPMode::ThreadSafe> CurrentPlayer(Player);
+	if (CurrentPlayer.IsValid())
+	{
+		TileNum = CurrentPlayer->GetTileNum();
+	}
+
+	return TileNum;
+}
 
 FTimespan FMediaPlayerFacade::GetTime() const
 {
diff --git a/Engine/Source/Runtime/MediaUtils/Public/MediaPlayerFacade.h b/Engine/Source/Runtime/MediaUtils/Public/MediaPlayerFacade.h
--- a/Engine/Source/Runtime/MediaUtils/Public/MediaPlayerFacade.h
+++ b/Engine/Source/Runtime/MediaUtils/Public/MediaPlayerFacade.h
@@ -312,6 +312,11 @@ public:
 	 */
 	TRangeSet<float> GetSupportedRates(bool Unthinned = true) const;
 
+	/**
+	 * Get the number of tiles in each frame.
+	 */
+	virtual FIntPoint GetTileNum() const;
+
 	/**
 	 * Get the media's current playback time.
 	 *
