ImgMedia: Fixed an issue with tile regions being calculated incorrectly.

diff --git a/Engine/Plugins/Media/ImgMedia/Source/ImgMedia/Private/Assets/ImgMediaMipMapInfo.cpp b/Engine/Plugins/Media/ImgMedia/Source/ImgMedia/Private/Assets/ImgMediaMipMapInfo.cpp
--- a/Engine/Plugins/Media/ImgMedia/Source/ImgMedia/Private/Assets/ImgMediaMipMapInfo.cpp
+++ b/Engine/Plugins/Media/ImgMedia/Source/ImgMedia/Private/Assets/ImgMediaMipMapInfo.cpp
@@ -36,13 +36,13 @@ FImgMediaTileSelection::FImgMediaTileSelection(int32 NumTilesX, int32 NumTilesY,
 {
 }
 
-FImgMediaTileSelection FImgMediaTileSelection::CreateForTargetMipLevel(const FIntPoint& BaseTileNum, int32 TargetMipLevel, bool bDefaultVisibility)
+FImgMediaTileSelection FImgMediaTileSelection::CreateForTargetMipLevel(const FIntPoint& MipZeroResolution, const FIntPoint& TileSize, int32 TargetMipLevel, bool bDefaultVisibility)
 {
 	ensure(TargetMipLevel >= 0);
 
 	const int MipLevelDiv = 1 << TargetMipLevel;
-	int32 NumTilesX = FMath::Max(1, FMath::CeilToInt(float(BaseTileNum.X) / MipLevelDiv));
-	int32 NumTilesY = FMath::Max(1, FMath::CeilToInt(float(BaseTileNum.Y) / MipLevelDiv));
+	int32 NumTilesX = FMath::Max(1, FMath::CeilToInt(float(MipZeroResolution.X / MipLevelDiv) / TileSize.X));
+	int32 NumTilesY = FMath::Max(1, FMath::CeilToInt(float(MipZeroResolution.Y / MipLevelDiv) / TileSize.Y));
 
 	return FImgMediaTileSelection(NumTilesX, NumTilesY, bDefaultVisibility);
 }
@@ -89,7 +89,8 @@ bool FImgMediaTileSelection::Contains(const FImgMediaTileSelection& Other) const
 
 void FImgMediaTileSelection::SetVisible(int32 TileCoordX, int32 TileCoordY)
 {
-	Tiles[ToIndex(TileCoordX, TileCoordY, Dimensions)] = true;
+	int Index = ToIndex(FMath::Min(TileCoordX, Dimensions.X - 1), FMath::Min(TileCoordY, Dimensions.Y - 1), Dimensions);
+	Tiles[Index] = true;
 	bCachedVisibleRegionDirty = true;
 }
 
@@ -273,7 +274,7 @@ void FImgMediaMipMapObjectInfo::CalculateVisibleTiles(const TArray<FImgMediaView
 	// We simply add fully visible regions for all mip levels
 	for (int32 MipLevel = 0; MipLevel < InSequenceInfo.NumMipLevels; ++MipLevel)
 	{
-		VisibleTiles.Add(MipLevel, FImgMediaTileSelection::CreateForTargetMipLevel(InSequenceInfo.TilingDescription.TileNum, MipLevel, true));
+		VisibleTiles.Add(MipLevel, FImgMediaTileSelection::CreateForTargetMipLevel(InSequenceInfo.Dim, InSequenceInfo.TilingDescription.TileSize, MipLevel, true));
 	}
 }
 
@@ -371,8 +372,8 @@ namespace {
 				int MipLevelDiv = 1 << MaxLevel;
 
 				FIntPoint CurrentNumTiles;
-				CurrentNumTiles.X = FMath::Max(1, FMath::CeilToInt(float(SequenceTileNum.X) / MipLevelDiv));
-				CurrentNumTiles.Y = FMath::Max(1, FMath::CeilToInt(float(SequenceTileNum.Y) / MipLevelDiv));
+				CurrentNumTiles.X = FMath::CeilToInt(float(InSequenceInfo.Dim.X / MipLevelDiv) / InSequenceInfo.TilingDescription.TileSize.X);
+				CurrentNumTiles.Y = FMath::CeilToInt(float(InSequenceInfo.Dim.Y / MipLevelDiv) / InSequenceInfo.TilingDescription.TileSize.Y);
 
 				// Starting with tiles at the highest mip level
 				TQueue<FIntVector> Tiles;
@@ -392,9 +393,10 @@ namespace {
 
 					int32 CurrentMipLevel = Tile.Z;
 					MipLevelDiv = 1 << CurrentMipLevel;
+
 					// Calculate the number of tiles at this mip level
-					CurrentNumTiles.X = FMath::Max(1, FMath::CeilToInt(float(SequenceTileNum.X) / MipLevelDiv));
-					CurrentNumTiles.Y = FMath::Max(1, FMath::CeilToInt(float(SequenceTileNum.Y) / MipLevelDiv));
+					CurrentNumTiles.X = FMath::Max(1, FMath::CeilToInt((float(InSequenceInfo.Dim.X) / MipLevelDiv) / InSequenceInfo.TilingDescription.TileSize.X));
+					CurrentNumTiles.Y = FMath::Max(1, FMath::CeilToInt((float(InSequenceInfo.Dim.Y) / MipLevelDiv) / InSequenceInfo.TilingDescription.TileSize.Y));
 
 					FVector2f CurrentPartialTileNum = InSequenceInfo.GetPartialTileNum(CurrentMipLevel);
 
@@ -631,10 +633,11 @@ namespace {
 								{
 									if (!VisibleTiles.Contains(Level))
 									{
-										VisibleTiles.Emplace(Level, FImgMediaTileSelection::CreateForTargetMipLevel(SequenceTileNum, Level, false));
+										VisibleTiles.Emplace(Level, FImgMediaTileSelection::CreateForTargetMipLevel(InSequenceInfo.Dim, InSequenceInfo.TilingDescription.TileSize, Level, false));
 									}
 
 									const int MipLevelDiv = 1 << Level;
+
 									VisibleTiles[Level].SetVisible(TileX / MipLevelDiv, TileY / MipLevelDiv);
 								}
 							}
diff --git a/Engine/Plugins/Media/ImgMedia/Source/ImgMedia/Private/Assets/ImgMediaMipMapInfo.h b/Engine/Plugins/Media/ImgMedia/Source/ImgMedia/Private/Assets/ImgMediaMipMapInfo.h
--- a/Engine/Plugins/Media/ImgMedia/Source/ImgMedia/Private/Assets/ImgMediaMipMapInfo.h
+++ b/Engine/Plugins/Media/ImgMedia/Source/ImgMedia/Private/Assets/ImgMediaMipMapInfo.h
@@ -68,12 +68,12 @@ struct FImgMediaTileSelection
 	/**
 	 * Create and initialize a new tile selection, adjusting tile counts for the specified higher mip level.
 	 *
-	 * @param BaseNumTilesX Horizontal tile count at mip level 0.
-	 * @param NumTilesY Vertical tile count at mip level 0.
+	 * @param MipZeroResolution Pixel resolution of Mip 0.
+	 * @param TileSize Dimensions of tiles in pixels.
 	 * @param TargetMipLevel Higher target mip level for the selection (usually 1 and above).
 	 * @param bDefaultVisibility Optional visibility for the entire region.
 	 */
-	static FImgMediaTileSelection CreateForTargetMipLevel(const FIntPoint& BaseTileNum, int32 TargetMipLevel, bool bDefaultVisibility = false);
+	static FImgMediaTileSelection CreateForTargetMipLevel(const FIntPoint& MipZeroResolution, const FIntPoint& TileSize, int32 TargetMipLevel, bool bDefaultVisibility = false);
 
 	FImgMediaTileSelection() = default;
 	~FImgMediaTileSelection() noexcept = default;
diff --git a/Engine/Plugins/Media/ImgMedia/Source/ImgMedia/Private/Loader/ImgMediaLoader.cpp b/Engine/Plugins/Media/ImgMedia/Source/ImgMedia/Private/Loader/ImgMediaLoader.cpp
--- a/Engine/Plugins/Media/ImgMedia/Source/ImgMedia/Private/Loader/ImgMediaLoader.cpp
+++ b/Engine/Plugins/Media/ImgMedia/Source/ImgMedia/Private/Loader/ImgMediaLoader.cpp
@@ -1385,7 +1385,7 @@ void FImgMediaLoader::GetDesiredMipTiles(int32 FrameIndex, TMap<int32, FImgMedia
 		//Fallback case where we activate all mips and tiles.
 		for (int32 MipLevel = 0; MipLevel < GetNumMipLevels(); ++MipLevel)
 		{
-			OutMipsAndTiles.Emplace(MipLevel, FImgMediaTileSelection::CreateForTargetMipLevel(TilingDescription.TileNum, MipLevel, true));
+			OutMipsAndTiles.Emplace(MipLevel, FImgMediaTileSelection::CreateForTargetMipLevel(SequenceDim, TilingDescription.TileSize, MipLevel, true));
 		}
 	}
 }
diff --git a/Engine/Plugins/Media/ImgMedia/Source/ImgMedia/Private/Readers/ExrImgMediaReader.cpp b/Engine/Plugins/Media/ImgMedia/Source/ImgMedia/Private/Readers/ExrImgMediaReader.cpp
--- a/Engine/Plugins/Media/ImgMedia/Source/ImgMedia/Private/Readers/ExrImgMediaReader.cpp
+++ b/Engine/Plugins/Media/ImgMedia/Source/ImgMedia/Private/Readers/ExrImgMediaReader.cpp
@@ -101,8 +101,13 @@ FExrImgMediaReader::EReadResult FExrImgMediaReader::ReadTiles
 	{
 		int64 CurrentBufferPos = 0;
 		TRACE_CPUPROFILER_EVENT_SCOPE_TEXT(*FString::Printf(TEXT("FExrImgMediaReader_ReadTilesCustom_ReadTiles")));
-		for (const FIntRect& TileRegion : TileRegions)
+		for (const FIntRect& RawTileRegion : TileRegions)
 		{
+			// This clamp is to make sure that tile region is not out of bounds in case the region wasn't calculated incorrectly for some reason.
+			const FIntRect& TileRegion = FIntRect(
+				FIntPoint(FMath::Max(RawTileRegion.Min.X, 0), FMath::Max(RawTileRegion.Min.Y, 0)),
+				FIntPoint(FMath::Min(RawTileRegion.Max.X, DimensionInTiles.X), FMath::Min(RawTileRegion.Max.Y, DimensionInTiles.Y)));
+
 			for (int32 TileRow = TileRegion.Min.Y; TileRow < TileRegion.Max.Y; TileRow++)
 			{
 				// Check to see if the frame was canceled.
