ImgMedia: Filter out files with unsupported extensions from image sequences.

#rnx

diff --git a/Engine/Plugins/Media/ImgMedia/Source/ImgMedia/Private/Loader/ImgMediaLoader.cpp b/Engine/Plugins/Media/ImgMedia/Source/ImgMedia/Private/Loader/ImgMediaLoader.cpp
--- a/Engine/Plugins/Media/ImgMedia/Source/ImgMedia/Private/Loader/ImgMediaLoader.cpp
+++ b/Engine/Plugins/Media/ImgMedia/Source/ImgMedia/Private/Loader/ImgMediaLoader.cpp
@@ -993,6 +993,9 @@ void FImgMediaLoader::FindFiles(const FString& SequencePath, TArray<FString>& Ou
 
 	UE_LOG(LogImgMedia, Verbose, TEXT("Loader %p: Found %i image files in %s"), this, FoundFiles.Num(), *SequencePath);
 
+	// Same list of extensions as FGenericImgMediaReader::LoadFrameImage() with the addition of exr
+	const TSet<FString> SupportedImageFileExtensions = { TEXT("exr"), TEXT("jpg"), TEXT("jpeg"), TEXT("png"), TEXT("bmp") };
+
 	TArray<FString> UnnumberedFiles;
 	TSortedMap<int32, FString> SortedFoundFiles;
 	SortedFoundFiles.Reserve(FoundFiles.Num());
@@ -1000,6 +1003,13 @@ void FImgMediaLoader::FindFiles(const FString& SequencePath, TArray<FString>& Ou
 	int32 FrameNumber;
 	for (FString& File : FoundFiles)
 	{
+		const FString FileExtension = FPaths::GetExtension(File).ToLower();
+
+		if (!SupportedImageFileExtensions.Contains(FileExtension))
+		{
+			continue;
+		}
+
 		if (GetNumberAtEndOfString(FrameNumber, File))
 		{
 			SortedFoundFiles.Add(FrameNumber, File);
