ImgMedia: Properly deprecating 5.0 functions for the media source, instead of the previous direct deletion.

#rnx

diff --git a/Engine/Plugins/Media/ImgMedia/Source/ImgMedia/Private/Assets/ImgMediaSource.cpp b/Engine/Plugins/Media/ImgMedia/Source/ImgMedia/Private/Assets/ImgMediaSource.cpp
--- a/Engine/Plugins/Media/ImgMedia/Source/ImgMedia/Private/Assets/ImgMediaSource.cpp
+++ b/Engine/Plugins/Media/ImgMedia/Source/ImgMedia/Private/Assets/ImgMediaSource.cpp
@@ -115,6 +115,11 @@ void UImgMediaSource::AddTargetObject(AActor* InActor)
 	MipMapInfo->AddObject(InActor);
 }
 
+void UImgMediaSource::AddTargetObject(AActor* InActor, float Width)
+{
+	AddTargetObject(InActor);
+}
+
 
 void UImgMediaSource::RemoveTargetObject(AActor* InActor)
 {
diff --git a/Engine/Plugins/Media/ImgMedia/Source/ImgMedia/Public/ImgMediaSource.h b/Engine/Plugins/Media/ImgMedia/Source/ImgMedia/Public/ImgMediaSource.h
--- a/Engine/Plugins/Media/ImgMedia/Source/ImgMedia/Public/ImgMediaSource.h
+++ b/Engine/Plugins/Media/ImgMedia/Source/ImgMedia/Public/ImgMediaSource.h
@@ -103,6 +103,15 @@ public:
 	UFUNCTION(BlueprintCallable, Category="ImgMedia|ImgMediaSource")
 	void SetSequencePath(const FString& Path);
 
+	UE_DEPRECATED(5.1, "AddGlobalCamera has been deprecated.")
+	UFUNCTION(BlueprintCallable, Category = "ImgMedia|ImgMediaSource", meta = (DeprecatedFunction, DeprecationMessage = "AddGlobalCamera has been deprecated."))
+	void AddGlobalCamera(AActor* InActor) { }
+
+
+	UE_DEPRECATED(5.1, "RemoveGlobalCamera has been deprecated.")
+	UFUNCTION(BlueprintCallable, Category = "ImgMedia|ImgMediaSource", meta = (DeprecatedFunction, DeprecationMessage = "RemoveGlobalCamera has been deprecated."))
+	void RemoveGlobalCamera(AActor* InActor) { }
+
 	/**
 	 * Set the path to the image sequence directory this source represents.
 	 *
@@ -120,6 +129,9 @@ public:
 	UFUNCTION(BlueprintCallable, Category = "ImgMedia|ImgMediaSource")
 	void AddTargetObject(AActor* InActor);
 
+	UE_DEPRECATED(5.1, "AddTargetObject has been deprecated, please use the variant with no Width argument.")
+	void AddTargetObject(AActor* InActor, float Width);
+
 	/**
 	 * This object is no longer using our img sequence.
 	 *
@@ -128,6 +140,10 @@ public:
 	UFUNCTION(BlueprintCallable, Category = "ImgMedia|ImgMediaSource")
 	void RemoveTargetObject(AActor* InActor);
 
+	UE_DEPRECATED(5.1, "SetMipLevelDistance has been deprecated.")
+	UFUNCTION(BlueprintCallable, Category = "ImgMedia|ImgMediaSource", meta = (DeprecatedFunction, DeprecationMessage = "SetMipLevelDistance has been deprecated."))
+	void SetMipLevelDistance(float Distance) {}
+
 	/**
 	 * Get our mipmap info object.
 	 *
