ImgMedia: Receive view information for mip/tile calculations at the end of a frame on the game thread.

This ensures that preview renders (nDisplay) do not invalidate our main view info.

#rnx

diff --git a/Engine/Plugins/Media/ImgMedia/Source/ImgMedia/Private/Assets/ImgMediaMipMapInfo.cpp b/Engine/Plugins/Media/ImgMedia/Source/ImgMedia/Private/Assets/ImgMediaMipMapInfo.cpp
--- a/Engine/Plugins/Media/ImgMedia/Source/ImgMedia/Private/Assets/ImgMediaMipMapInfo.cpp
+++ b/Engine/Plugins/Media/ImgMedia/Source/ImgMedia/Private/Assets/ImgMediaMipMapInfo.cpp
@@ -13,6 +13,7 @@
 #include "Engine/StaticMesh.h"
 #include "GameFramework/Actor.h"
 #include "Math/UnrealMathUtility.h"
+#include "Misc/CoreDelegates.h"
 
 #if WITH_EDITOR
 #include "Editor.h"
@@ -665,6 +666,7 @@ namespace {
 
 FImgMediaMipMapInfo::FImgMediaMipMapInfo()
 	: bIsCacheValid(false)
+	, OnEndFrameHandle()
 {
 }
 
@@ -696,6 +698,8 @@ void FImgMediaMipMapInfo::AddObject(AActor* InActor, float InMipMapBias, EMediaT
 				break;
 			}
 		}
+
+		SubscribeEndFrame();
 	}
 }
 
@@ -715,6 +719,11 @@ void FImgMediaMipMapInfo::RemoveObject(AActor* InActor)
 				{
 					Objects.RemoveAtSwap(Index);
 					delete Info;
+
+					if (Objects.IsEmpty())
+					{
+						UnsubscribeEndFrame();
+					}
 					break;
 				}
 			}
@@ -772,6 +781,8 @@ void FImgMediaMipMapInfo::ClearAllObjects()
 	}
 
 	Objects.Empty();
+
+	UnsubscribeEndFrame();
 }
 
 
@@ -835,46 +846,76 @@ bool FImgMediaMipMapInfo::HasObjects() const
 
 void FImgMediaMipMapInfo::Tick(float DeltaTime)
 {
-	FScopeLock Lock(&InfoCriticalSection);
-
-	// Let the cache update this frame.
-	bIsCacheValid = false;
+	if (!HasObjects())
+	{
+		return;
+	}
 
-	if (HasObjects())
+	// Display debug?
+	if (CVarImgMediaMipMapDebugEnable.GetValueOnGameThread() && GEngine != nullptr)
 	{
-		const TSharedPtr<FImgMediaSceneViewExtension, ESPMode::ThreadSafe>& SVE = IImgMediaModule::Get().GetSceneViewExtension();
-		if (SVE.IsValid())
+		FScopeLock Lock(&InfoCriticalSection);
+
+		TSet<int32> VisibleMips;
+		int32 NumVisibleTiles = 0;
+
+		for (const auto& MipTiles : CachedVisibleTiles)
 		{
-			ViewInfos = SVE->GetViewInfos();
+			VisibleMips.Add(MipTiles.Key);
+
+			const FImgMediaTileSelection& TileSelection = MipTiles.Value;
+			NumVisibleTiles += TileSelection.NumVisibleTiles();
 		}
 
-		// Display debug?
-		if (CVarImgMediaMipMapDebugEnable.GetValueOnGameThread() && GEngine != nullptr)
+		if (VisibleMips.Num() > 0)
 		{
-			TSet<int32> VisibleMips;
-			int32 NumVisibleTiles = 0;
+			auto VisibleMipsIt = VisibleMips.CreateConstIterator();
+			FString Mips = FString::FromInt(*VisibleMipsIt);
 
-			for (const auto& MipTiles : CachedVisibleTiles)
+			for (++VisibleMipsIt; VisibleMipsIt; ++VisibleMipsIt)
 			{
-				VisibleMips.Add(MipTiles.Key);
-
-				const FImgMediaTileSelection& TileSelection = MipTiles.Value;
-				NumVisibleTiles += TileSelection.NumVisibleTiles();
+				Mips += FString(TEXT(", ")) + FString::FromInt(*VisibleMipsIt);
 			}
 
-			if (VisibleMips.Num() > 0)
-			{
-				auto VisibleMipsIt = VisibleMips.CreateConstIterator();
-				FString Mips = FString::FromInt(*VisibleMipsIt);
-				
-				for (++VisibleMipsIt; VisibleMipsIt; ++VisibleMipsIt)
-				{
-					Mips += FString(TEXT(", ")) + FString::FromInt(*VisibleMipsIt);
-				}
+			GEngine->AddOnScreenDebugMessage(-1, 0.0f, FColor::Yellow, *FString::Printf(TEXT("%s Mip Level(s): [%s]"), *SequenceInfo.Name.ToString(), *Mips));
+			GEngine->AddOnScreenDebugMessage(-1, 0.0f, FColor::Yellow, *FString::Printf(TEXT("%s Num Tile(s): %d"), *SequenceInfo.Name.ToString(), NumVisibleTiles));
+		}
+	}
+}
 
-				GEngine->AddOnScreenDebugMessage(-1, 0.0f, FColor::Yellow, *FString::Printf(TEXT("%s Mip Level(s): [%s]"), *SequenceInfo.Name.ToString(), *Mips));
-				GEngine->AddOnScreenDebugMessage(-1, 0.0f, FColor::Yellow, *FString::Printf(TEXT("%s Num Tile(s): %d"), *SequenceInfo.Name.ToString(), NumVisibleTiles));
-			}
+void FImgMediaMipMapInfo::SubscribeEndFrame()
+{
+	if (!OnEndFrameHandle.IsValid())
+	{
+		OnEndFrameHandle = FCoreDelegates::OnEndFrame.AddRaw(this, &FImgMediaMipMapInfo::OnEndFrame);
+	}
+}
+
+void FImgMediaMipMapInfo::UnsubscribeEndFrame()
+{
+	if (OnEndFrameHandle.IsValid())
+	{
+		FCoreDelegates::OnEndFrame.Remove(OnEndFrameHandle);
+		OnEndFrameHandle.Reset();
+
+		bIsCacheValid = false;
+	}
+}
+
+void FImgMediaMipMapInfo::OnEndFrame()
+{
+	TRACE_CPUPROFILER_EVENT_SCOPE(FImgMediaMipMapInfo::OnEndFrame);
+	FScopeLock Lock(&InfoCriticalSection);
+
+	// Let the cache update.
+	bIsCacheValid = false;
+
+	if (HasObjects())
+	{
+		const TSharedPtr<FImgMediaSceneViewExtension, ESPMode::ThreadSafe>& SVE = IImgMediaModule::Get().GetSceneViewExtension();
+		if (SVE.IsValid())
+		{
+			ViewInfos = SVE->GetViewInfos();
 		}
 	}
 }
diff --git a/Engine/Plugins/Media/ImgMedia/Source/ImgMedia/Private/Assets/ImgMediaMipMapInfo.h b/Engine/Plugins/Media/ImgMedia/Source/ImgMedia/Private/Assets/ImgMediaMipMapInfo.h
--- a/Engine/Plugins/Media/ImgMedia/Source/ImgMedia/Private/Assets/ImgMediaMipMapInfo.h
+++ b/Engine/Plugins/Media/ImgMedia/Source/ImgMedia/Private/Assets/ImgMediaMipMapInfo.h
@@ -292,6 +292,12 @@ public:
 
 protected:
 
+	void SubscribeEndFrame();
+	void UnsubscribeEndFrame();
+
+	/** End frame callback registered with core delegates to receive view infos. */
+	void OnEndFrame();
+
 	/** Array of objects that are using our img sequence. */
 	TArray<FImgMediaMipMapObjectInfo*> Objects;
 
@@ -312,4 +318,7 @@ protected:
 
 	/** True if the cached mipmap data has been calculated this frame. */
 	bool bIsCacheValid;
+
+	/** OnEndFrame core delegate handle. */
+	FDelegateHandle OnEndFrameHandle;
 };
diff --git a/Engine/Plugins/Media/ImgMedia/Source/ImgMedia/Private/ImgMediaSceneViewExtension.cpp b/Engine/Plugins/Media/ImgMedia/Source/ImgMedia/Private/ImgMediaSceneViewExtension.cpp
--- a/Engine/Plugins/Media/ImgMedia/Source/ImgMedia/Private/ImgMediaSceneViewExtension.cpp
+++ b/Engine/Plugins/Media/ImgMedia/Source/ImgMedia/Private/ImgMediaSceneViewExtension.cpp
@@ -28,6 +28,8 @@ void FImgMediaSceneViewExtension::SetupView(FSceneViewFamily& InViewFamily, FSce
 
 void FImgMediaSceneViewExtension::BeginRenderViewFamily(FSceneViewFamily& InViewFamily)
 {
+	TRACE_CPUPROFILER_EVENT_SCOPE(FImgMediaSceneViewExtension::BeginRenderViewFamily);
+
 	if (LastFrameNumber != InViewFamily.FrameNumber)
 	{
 		CachedViewInfos.Reset();
