ImgMedia: Remove includes to private headers from the shader public header.

#rnx

diff --git a/Engine/Plugins/Media/ImgMedia/Source/ExrReaderGpu/ExrReaderGpu.Build.cs b/Engine/Plugins/Media/ImgMedia/Source/ExrReaderGpu/ExrReaderGpu.Build.cs
--- a/Engine/Plugins/Media/ImgMedia/Source/ExrReaderGpu/ExrReaderGpu.Build.cs
+++ b/Engine/Plugins/Media/ImgMedia/Source/ExrReaderGpu/ExrReaderGpu.Build.cs
@@ -13,22 +13,11 @@ namespace UnrealBuildTool.Rules
 				new string[] {
 					"Core",
 					"CoreUObject",
-					"TimeManagement",
-					"RHI",
 					"RenderCore",
-					"Renderer",
-					"Slate",
-					"SlateCore",
 					"Projects",
 					"Engine",
 				});
 
-			if (Target.Type == TargetType.Editor)
-			{
-				PublicDependencyModuleNames.Add("UnrealEd");
-				PrivateDependencyModuleNames.Add("UnrealEd");
-			}
-
 			if (Target.Platform.IsInGroup(UnrealPlatformGroup.Windows))
 			{
 				AddEngineThirdPartyPrivateStaticDependencies(Target, "Imath");
diff --git a/Engine/Plugins/Media/ImgMedia/Source/ExrReaderGpu/Public/ExrSwizzlingShader.h b/Engine/Plugins/Media/ImgMedia/Source/ExrReaderGpu/Public/ExrSwizzlingShader.h
--- a/Engine/Plugins/Media/ImgMedia/Source/ExrReaderGpu/Public/ExrSwizzlingShader.h
+++ b/Engine/Plugins/Media/ImgMedia/Source/ExrReaderGpu/Public/ExrSwizzlingShader.h
@@ -5,16 +5,9 @@
 #if defined(PLATFORM_WINDOWS) && PLATFORM_WINDOWS
 
 #include "CoreMinimal.h"
-#include "UObject/ConstructorHelpers.h"
-#include "RHI.h"
-#include "Runtime/RenderCore/Public/ShaderParameterUtils.h"
-#include "Runtime/RenderCore/Public/RenderResource.h"
-#include "Runtime/RenderCore/Public/RenderGraphResources.h"
-
-// FScreenPassTextureViewportParameters and FScreenPassTextureInput
-#include "Runtime/Renderer/Private/ScreenPass.h"
-#include "Runtime/Renderer/Private/SceneTextureParameters.h"
-
+#include "GlobalShader.h"
+#include "ShaderParameterUtils.h"
+#include "ShaderParameterStruct.h"
 
 // The vertex shader used by DrawScreenPass to draw a rectangle.
 class EXRREADERGPU_API FExrSwizzleVS : public FGlobalShader
