ImgMedia: Cache an empty frame when no tiles are visible to fix blocking playback.

#rnx

diff --git a/Engine/Plugins/Media/ImgMedia/Source/ImgMedia/Private/GlobalCache/ImgMediaGlobalCache.cpp b/Engine/Plugins/Media/ImgMedia/Source/ImgMedia/Private/GlobalCache/ImgMediaGlobalCache.cpp
--- a/Engine/Plugins/Media/ImgMedia/Source/ImgMedia/Private/GlobalCache/ImgMediaGlobalCache.cpp
+++ b/Engine/Plugins/Media/ImgMedia/Source/ImgMedia/Private/GlobalCache/ImgMediaGlobalCache.cpp
@@ -127,6 +127,13 @@ void FImgMediaGlobalCache::AddFrame(const FString& FileName, const FName& Sequen
 	}
 }
 
+bool FImgMediaGlobalCache::Contains(const FName& Sequence, int32 Index)
+{
+	FScopeLock Lock(&CriticalSection);
+
+	return MapFrameToEntry.Contains(TPair<FName, int32>(Sequence, Index));
+}
+
 TSharedPtr<FImgMediaFrame, ESPMode::ThreadSafe>* FImgMediaGlobalCache::FindAndTouch(const FName& Sequence, int32 Index)
 {
 	FScopeLock Lock(&CriticalSection);
diff --git a/Engine/Plugins/Media/ImgMedia/Source/ImgMedia/Private/Loader/ImgMediaLoader.cpp b/Engine/Plugins/Media/ImgMedia/Source/ImgMedia/Private/Loader/ImgMediaLoader.cpp
--- a/Engine/Plugins/Media/ImgMedia/Source/ImgMedia/Private/Loader/ImgMediaLoader.cpp
+++ b/Engine/Plugins/Media/ImgMedia/Source/ImgMedia/Private/Loader/ImgMediaLoader.cpp
@@ -52,15 +52,21 @@ namespace ImgMediaLoader
 	}
 	
 	// Check if the existing tiles contain all of the requested ones. (Is existing a superset of requested?)
-	bool ContainsMipTiles(const TMap<int32, FImgMediaTileSelection>& ExistingTiles, const TMap<int32, FImgMediaTileSelection>& RequestedTiles)
+	bool ContainsMipTiles(const TSharedPtr<FImgMediaFrame, ESPMode::ThreadSafe>& ExistingFrame, const TMap<int32, FImgMediaTileSelection>& RequestedTiles)
 	{
 		for (auto Iter = RequestedTiles.CreateConstIterator(); Iter; ++Iter)
 		{
 			const int32 RequestedMipLevel = Iter.Key();
 			const FImgMediaTileSelection& RequestedSelection = Iter.Value();
 
-			if (const FImgMediaTileSelection* ExistingSelection = ExistingTiles.Find(RequestedMipLevel))
+			if (const FImgMediaTileSelection* ExistingSelection = ExistingFrame->MipTilesPresent.Find(RequestedMipLevel))
 			{
+				if (ExistingSelection->GetDimensions().GetMax() < RequestedSelection.GetDimensions().GetMax())
+				{
+					// If an empty 1x1 frame was cached but we're now requesting a fully tiled frame.
+					return false;
+				}
+
 				if (!ExistingSelection->Contains(RequestedSelection))
 				{
 					// Requested tile selection is not present.
@@ -681,19 +687,27 @@ IQueuedWork* FImgMediaLoader::GetWork()
 {
 	FScopeLock Lock(&CriticalSection);
 
-	TMap<int32, FImgMediaTileSelection> DesiredMipsAndTiles;
-	int32 FrameNumber = INDEX_NONE;
-
-	// Find a visible pending frame.
-	while (!PendingFrameNumbers.IsEmpty() && DesiredMipsAndTiles.IsEmpty())
+	if (PendingFrameNumbers.IsEmpty())
 	{
-		FrameNumber = PendingFrameNumbers.Pop(false);
-		GetDesiredMipTiles(FrameNumber, DesiredMipsAndTiles);
+		return nullptr;
 	}
 
-	// No selection is visible (or pending frames was empty), so we don't queue any work.
+	int32 FrameNumber = PendingFrameNumbers.Pop(false);
+
+	TMap<int32, FImgMediaTileSelection> DesiredMipsAndTiles;
+	GetDesiredMipTiles(FrameNumber, DesiredMipsAndTiles);
+
 	if (DesiredMipsAndTiles.IsEmpty())
 	{
+		// Still provide the cache with an empty frame to prevent blocking playback from stalling.
+		bool bIsFrameAlreadyCached = UseGlobalCache ? GlobalCache->Contains(SequenceName, FrameNumber) : Frames.Contains(FrameNumber);
+
+		if (!bIsFrameAlreadyCached)
+		{
+			AddEmptyFrame(FrameNumber);
+		}
+
+		// No selection was visible, so we don't queue any work.
 		return nullptr;
 	}
 
@@ -1292,7 +1306,7 @@ void FImgMediaLoader::Update(int32 PlayHeadFrame, float PlayRate, bool Loop)
 			TMap<int32, FImgMediaTileSelection> DesiredMipsAndTiles;
 			GetDesiredMipTiles(FrameNumber, DesiredMipsAndTiles);
 
-			NeedFrame = !ImgMediaLoader::ContainsMipTiles((*FramePtr)->MipTilesPresent, DesiredMipsAndTiles);
+			NeedFrame = !ImgMediaLoader::ContainsMipTiles(*FramePtr, DesiredMipsAndTiles);
 		}
 		
 		if ((NeedFrame) && !QueuedFrameNumbers.Contains(FrameNumber))
@@ -1331,7 +1345,7 @@ void FImgMediaLoader::AddEmptyFrame(int32 FrameNumber)
 	Frame->Info.NumChannels = NumChannels;
 	Frame->Info.bHasTiles = false;
 	Frame->Format = EMediaTextureSampleFormat::FloatRGB;
-	Frame->Stride = Frame ->Info.Dim.X * PixelSize;
+	Frame->Stride = Frame->Info.Dim.X * PixelSize;
 	for (int32 Level = 0; Level < GetNumMipLevels(); ++Level)
 	{
 		Frame->MipTilesPresent.Emplace(Level, FImgMediaTileSelection(1, 1, true));
diff --git a/Engine/Plugins/Media/ImgMedia/Source/ImgMedia/Public/ImgMediaGlobalCache.h b/Engine/Plugins/Media/ImgMedia/Source/ImgMedia/Public/ImgMediaGlobalCache.h
--- a/Engine/Plugins/Media/ImgMedia/Source/ImgMedia/Public/ImgMediaGlobalCache.h
+++ b/Engine/Plugins/Media/ImgMedia/Source/ImgMedia/Public/ImgMediaGlobalCache.h
@@ -50,6 +50,14 @@ public:
 	 */
 	void AddFrame(const FString& FileName, const FName& Sequence, int32 Index, const TSharedPtr<FImgMediaFrame, ESPMode::ThreadSafe>& Frame, bool HasMipMaps);
 
+	/**
+	 * Check whether the entry with the specified sequence and index is present.
+	 *
+	 * @param	Sequence		Indentifying name of this sequence.
+	 * @param	Index			Index of frame to touch.
+	 */
+	bool Contains(const FName& Sequence, int32 Index);
+
 	/**
 	 * Find the entry with the specified sequence and index and mark it as the most recently used.
 	 *
