ImgMedia: Fixed numeric overflow in process exr.
#rnx

diff --git a/Engine/Plugins/Media/ImgMedia/Source/ImgMediaEditor/Private/Widgets/SImgMediaProcessEXR.cpp b/Engine/Plugins/Media/ImgMedia/Source/ImgMediaEditor/Private/Widgets/SImgMediaProcessEXR.cpp
--- a/Engine/Plugins/Media/ImgMedia/Source/ImgMediaEditor/Private/Widgets/SImgMediaProcessEXR.cpp
+++ b/Engine/Plugins/Media/ImgMedia/Source/ImgMediaEditor/Private/Widgets/SImgMediaProcessEXR.cpp
@@ -704,7 +704,7 @@ void SImgMediaProcessEXR::TintData(uint8* SourceData, TArray64<uint8>& DestArray
 
 	// Set up destination buffer.
 	DestArray.Reset();
-	DestArray.AddUninitialized(InWidth * InHeight * InNumChannels * 2);
+	DestArray.AddUninitialized(((int64)InWidth) * ((int64)InHeight) * InNumChannels * 2);
 	FFloat16* Buffer = (FFloat16*)SourceData;
 	FFloat16* OutBuffer = (FFloat16*)DestArray.GetData();
 		
