ImgMedia:
- ExrReaderGpu is now capable of upscaling lower level mips.
- Use upload heap cvar now needs to be set at the start of the engine.

diff --git a/Engine/Plugins/Media/ImgMedia/Source/ImgMedia/Private/Readers/ExrImgMediaReader.h b/Engine/Plugins/Media/ImgMedia/Source/ImgMedia/Private/Readers/ExrImgMediaReader.h
--- a/Engine/Plugins/Media/ImgMedia/Source/ImgMedia/Private/Readers/ExrImgMediaReader.h
+++ b/Engine/Plugins/Media/ImgMedia/Source/ImgMedia/Private/Readers/ExrImgMediaReader.h
@@ -87,6 +87,9 @@ protected:
 
 		/** Indicates if mips stored in individual files.*/
 		bool bMipsInSeparateFiles;
+
+		/** A lower quality mip will be upscaled if value is 0 or above. At 0 highest quality mip will always be read fully. */
+		int32 UpscaleMip = -1;
 	};
 
 	/**
diff --git a/Engine/Plugins/Media/ImgMedia/Source/ImgMedia/Private/Readers/ExrImgMediaReaderGpu.cpp b/Engine/Plugins/Media/ImgMedia/Source/ImgMedia/Private/Readers/ExrImgMediaReaderGpu.cpp
--- a/Engine/Plugins/Media/ImgMedia/Source/ImgMedia/Private/Readers/ExrImgMediaReaderGpu.cpp
+++ b/Engine/Plugins/Media/ImgMedia/Source/ImgMedia/Private/Readers/ExrImgMediaReaderGpu.cpp
@@ -28,13 +28,25 @@
 #include "ID3D12DynamicRHI.h"
 
 
-DECLARE_GPU_STAT_NAMED(ExrImgMediaReaderGpu, TEXT("ExrImgMediaReaderGpu"));
-DECLARE_GPU_STAT_NAMED(ExrImgMediaReaderGpu_MipRender, TEXT("ExrImgMediaReaderGpu_MipRender"));
+DECLARE_GPU_STAT_NAMED(ExrImgMediaReaderGpu, TEXT("ExrImgGpu"));
+DECLARE_GPU_STAT_NAMED(ExrImgMediaReaderGpu_MipRender, TEXT("ExrImgGpu.MipRender"));
+DECLARE_GPU_STAT_NAMED(ExrImgMediaReaderGpu_MipUpscale, TEXT("ExrImgGpu.MipRender.MipUpscale"));
 
 static TAutoConsoleVariable<bool> CVarExrReaderUseUploadHeap(
 	TEXT("r.ExrReaderGPU.UseUploadHeap"),
 	true,
-	TEXT("Utilizes upload heap and copies raw exr buffer asynchronously.\n"));
+	TEXT("Utilizes upload heap and copies raw exr buffer asynchronously.\n\
+			Requires a restart of the engine."));
+
+static TAutoConsoleVariable<int32> CVarUpscaleLowQualityMip(
+	TEXT("r.ExrReaderGPU.UpscaleHigherLevelMip"),
+	-1,
+	TEXT("Upscales lower quality mips into higher quality (EX: upscaling mip 3 into mip 0, 1, 2).\n\
+			This will cover unread black regions of EXR texture with the lower quality mips.\n\
+			This will clamp to the lowest quality mip available and disabled by default (-1)"));
+
+/** This is to force user to restart after they've changed this setting. */
+static bool bUseUploadHeap = CVarExrReaderUseUploadHeap.GetValueOnAnyThread();
 
 namespace {
 
@@ -115,6 +127,115 @@ FExrImgMediaReaderGpu::~FExrImgMediaReaderGpu()
 	}
 }
 
+FExrImgMediaReader::EReadResult FExrImgMediaReaderGpu::ReadMip
+	( const int32 CurrentMipLevel
+	, const FImgMediaTileSelection& CurrentTileSelection
+	, TSharedPtr<FImgMediaFrame, ESPMode::ThreadSafe> OutFrame
+	, TSharedPtr<FSampleConverterParameters> ConverterParams
+	, FExrMediaTextureSampleConverter* SampleConverter
+	, const FString& ImagePath
+	, bool bHasTiles)
+{
+
+	TRACE_CPUPROFILER_EVENT_SCOPE_TEXT(*FString::Printf(TEXT("ExrReaderGpu.ReadMip %d"), CurrentMipLevel));
+
+	// Next mip level.
+	int MipLevelDiv = 1 << CurrentMipLevel;
+	FIntPoint CurrentMipDim = ConverterParams->FullResolution / MipLevelDiv;
+
+	const SIZE_T BufferSize = GetBufferSize(CurrentMipDim, ConverterParams->FrameInfo.NumChannels, bHasTiles, OutFrame->Info.NumTiles / MipLevelDiv, ConverterParams->bCustomExr);
+	FStructuredBufferPoolItemSharedPtr BufferData = SampleConverter->GetMipLevelBuffer(CurrentMipLevel);
+
+	if (!BufferData.IsValid())
+	{
+		BufferData = AllocateGpuBufferFromPool(BufferSize);
+		SampleConverter->SetMipLevelBuffer(CurrentMipLevel, BufferData);
+	}
+
+	uint16* MipDataPtr = static_cast<uint16*>(BufferData->UploadBufferMapped);
+
+	EReadResult ReadResult = Fail;
+
+	if (FPaths::FileExists(ImagePath))
+	{
+		TArray<UE::Math::TIntPoint<int64>> BufferRegionsToCopy;
+		// read frame data
+		if (bHasTiles || ConverterParams->bCustomExr)
+		{
+			TArray<FIntRect> TileRegionsToRead;
+			TArray<FIntRect> TileRegionsToRender;
+			{
+				TRACE_CPUPROFILER_EVENT_SCOPE_TEXT(*FString::Printf(TEXT("ExrReaderGpu.CalculateRegions %d"), CurrentMipLevel));
+
+				if (const FImgMediaTileSelection* CachedSelection = OutFrame->MipTilesPresent.Find(CurrentMipLevel))
+				{
+					TileRegionsToRead = CachedSelection->GetVisibleRegions(&CurrentTileSelection);
+					TileRegionsToRender = CurrentTileSelection.GetVisibleRegions();
+				}
+				else
+				{
+					TileRegionsToRead = TileRegionsToRender = CurrentTileSelection.GetVisibleRegions();
+				}
+			}
+
+			TArray<FIntRect>& Viewports = ConverterParams->Viewports.Add(CurrentMipLevel);
+			for (const FIntRect& TileRegion : TileRegionsToRender)
+			{
+				FIntRect Viewport;
+				Viewport.Min = FIntPoint(ConverterParams->TileDimWithBorders.X * TileRegion.Min.X, ConverterParams->TileDimWithBorders.Y * TileRegion.Min.Y);
+				Viewport.Max = FIntPoint(ConverterParams->TileDimWithBorders.X * TileRegion.Max.X, ConverterParams->TileDimWithBorders.Y * TileRegion.Max.Y);
+				Viewport.Clip(FIntRect(FIntPoint::ZeroValue, CurrentMipDim));
+				Viewports.Add(MoveTemp(Viewport));
+			}
+
+			ReadResult = ReadTiles(MipDataPtr, BufferSize, ImagePath, ConverterParams->FrameId, TileRegionsToRead, ConverterParams, CurrentMipLevel, BufferRegionsToCopy);
+
+			for (const FIntRect& Region : TileRegionsToRead)
+			{
+				OutFrame->NumTilesRead += Region.Area();
+			}
+		}
+		else
+		{
+			TArray<FIntRect>& Viewports = ConverterParams->Viewports.Add(CurrentMipLevel);
+			FIntRect Viewport;
+
+			Viewport.Min = FIntPoint(0, 0);
+			Viewport.Max = CurrentMipDim;
+			Viewports.Add(MoveTemp(Viewport));
+
+			ReadResult = ReadInChunks(MipDataPtr, ImagePath, ConverterParams->FrameId, CurrentMipDim, BufferSize);
+			OutFrame->NumTilesRead++;
+		}
+
+		if (ReadResult == Success && bUseUploadHeap)
+		{
+			ENQUEUE_RENDER_COMMAND(CopyFromUploadBuffer)([SampleConverter, ConverterParams, BufferData, BufferRegionsToCopy](FRHICommandListImmediate& RHICmdList)
+				{
+					SCOPED_DRAW_EVENT(RHICmdList, FExrImgMediaReaderGpu_CopyBuffers);
+					TRACE_CPUPROFILER_EVENT_SCOPE_TEXT(*FString::Printf(TEXT("ExrReaderGpu.StartCopy %d"), ConverterParams->FrameId));
+					if (BufferRegionsToCopy.IsEmpty())
+					{
+						RHICmdList.CopyBuffer(BufferData->UploadBufferRef, BufferData->ShaderAccessBufferRef);
+					}
+					else
+					{
+						for (UE::Math::TIntPoint<int64> Region : BufferRegionsToCopy)
+						{
+							RHICmdList.CopyBufferRegion(BufferData->ShaderAccessBufferRef, Region.X, BufferData->UploadBufferRef, Region.X, Region.Y);
+						}
+					}
+				});
+		}
+	}
+	else
+	{
+		UE_LOG(LogImgMedia, Error, TEXT("Could not load %s"), *ImagePath);
+		return Fail;
+	}
+
+	return ReadResult;
+}
 
 /* FExrImgMediaReaderGpu interface
  *****************************************************************************/
@@ -169,17 +290,30 @@ bool FExrImgMediaReaderGpu::ReadFrame(int32 FrameId, const TMap<int32, FImgMedia
 	}
 	SampleConverter = static_cast<FExrMediaTextureSampleConverter*>(OutFrame->SampleConverter.Get());
 
+
 	{
-		// Loop over all mips.
+		// Force mip level to be upscaled to all higher quality mips.
 		FIntPoint CurrentMipDim = ConverterParams->FullResolution;
-
-		for (const TPair<int32, FImgMediaTileSelection>& TilesPerMip : InMipTiles)
+		TMap<int32, FImgMediaTileSelection> InMipTilesCopy = InMipTiles;
+		int32 MipToUpscale = FMath::Clamp(CVarUpscaleLowQualityMip.GetValueOnAnyThread(), -1, ConverterParams->NumMipLevels - 1);
+		ConverterParams->UpscaleMip = MipToUpscale;
+		if (MipToUpscale >= 0)
 		{
+			FImgMediaTileSelection FullSelection = FImgMediaTileSelection::CreateForTargetMipLevel(ConverterParams->FullResolution, TileDim, MipToUpscale, true);
+			if (InMipTilesCopy.Contains(MipToUpscale))
+			{
+				InMipTilesCopy[MipToUpscale] = FullSelection;
+			}
+			else
+			{
+				InMipTilesCopy.Add(MipToUpscale, FullSelection);
+			}
+		}
 
+		// Loop over all mips.
+		for (const TPair<int32, FImgMediaTileSelection>& TilesPerMip : InMipTilesCopy)
+		{
 			const int32 CurrentMipLevel = TilesPerMip.Key;
-
-			TRACE_CPUPROFILER_EVENT_SCOPE_TEXT(*FString::Printf(TEXT("ExrReaderGpu.ReadMip %d"), CurrentMipLevel));
-
 			const FImgMediaTileSelection& CurrentTileSelection = TilesPerMip.Value;
 
 			bool ReadThisMip = true;
@@ -188,112 +322,17 @@ bool FExrImgMediaReaderGpu::ReadFrame(int32 FrameId, const TMap<int32, FImgMedia
 			{
 				ReadThisMip = !CachedSelection->Contains(CurrentTileSelection);
 			}
-			
+
 			if (!ReadThisMip)
 			{
 				continue;
 			}
-
-			// Next mip level.
-			int MipLevelDiv = 1 << CurrentMipLevel;
-			CurrentMipDim = ConverterParams->FullResolution / MipLevelDiv;
-
-
-			const SIZE_T BufferSize = GetBufferSize(CurrentMipDim, ConverterParams->FrameInfo.NumChannels, bHasTiles, OutFrame->Info.NumTiles / MipLevelDiv, ConverterParams->bCustomExr);
-			FStructuredBufferPoolItemSharedPtr BufferData = SampleConverter->GetMipLevelBuffer(CurrentMipLevel);
-
-			if (!BufferData.IsValid())
-			{
-				BufferData = AllocateGpuBufferFromPool(BufferSize);
-				SampleConverter->SetMipLevelBuffer(CurrentMipLevel, BufferData);
-			}
-
-			uint16* MipDataPtr = static_cast<uint16*>(BufferData->UploadBufferMapped);
-
-			// Get highest resolution mip level path.
-			FString ImagePath = Loader->GetImagePath(FrameId, ConverterParams->bMipsInSeparateFiles ? CurrentMipLevel : 0);
-				
 			EReadResult ReadResult = Fail;
 
-			if (FPaths::FileExists(ImagePath))
-			{
-				TArray<UE::Math::TIntPoint<int64>> BufferRegionsToCopy;
-				// read frame data
-				if (bHasTiles || ConverterParams->bCustomExr)
-				{
-					TArray<FIntRect> TileRegionsToRead;
-					TArray<FIntRect> TileRegionsToRender;
-					{
-						TRACE_CPUPROFILER_EVENT_SCOPE_TEXT(*FString::Printf(TEXT("ExrReaderGpu.CalculateRegions %d"), CurrentMipLevel));
-
-						if (const FImgMediaTileSelection* CachedSelection = OutFrame->MipTilesPresent.Find(CurrentMipLevel))
-						{
-							TileRegionsToRead = CachedSelection->GetVisibleRegions(&CurrentTileSelection);
-							TileRegionsToRender = CurrentTileSelection.GetVisibleRegions();
-						}
-						else
-						{
-							TileRegionsToRead = TileRegionsToRender = CurrentTileSelection.GetVisibleRegions();
-						}
-					}
-
-					TArray<FIntRect>& Viewports = ConverterParams->Viewports.Add(CurrentMipLevel);
-					for (const FIntRect& TileRegion : TileRegionsToRender)
-					{
-						FIntRect Viewport;
-						Viewport.Min = FIntPoint(ConverterParams->TileDimWithBorders.X * TileRegion.Min.X, ConverterParams->TileDimWithBorders.Y * TileRegion.Min.Y);
-						Viewport.Max = FIntPoint(ConverterParams->TileDimWithBorders.X * TileRegion.Max.X, ConverterParams->TileDimWithBorders.Y * TileRegion.Max.Y);
-						Viewport.Clip(FIntRect(FIntPoint::ZeroValue, CurrentMipDim));
-						Viewports.Add(MoveTemp(Viewport));
-					}
-
-					ReadResult = ReadTiles(MipDataPtr, BufferSize, ImagePath, FrameId, TileRegionsToRead, ConverterParams, CurrentMipLevel, BufferRegionsToCopy);
-
-					for (const FIntRect& Region : TileRegionsToRead)
-					{
-						OutFrame->NumTilesRead += Region.Area();
-					}
-				}
-				else
-				{
-					TArray<FIntRect>& Viewports = ConverterParams->Viewports.Add(CurrentMipLevel);
-					FIntRect Viewport;
-						
-					Viewport.Min = FIntPoint(0,0);
-					Viewport.Max = CurrentMipDim;
-					Viewports.Add(MoveTemp(Viewport));
-
-					ReadResult = ReadInChunks(MipDataPtr, ImagePath, FrameId, CurrentMipDim, BufferSize);
-					OutFrame->NumTilesRead++;
-				}
+			// Get highest resolution mip level path.
+			FString ImagePath = Loader->GetImagePath(ConverterParams->FrameId, ConverterParams->bMipsInSeparateFiles ? CurrentMipLevel : 0);
 
-				if (ReadResult == Success && CVarExrReaderUseUploadHeap.GetValueOnAnyThread())
-				{
-					TArray<int32> InMipLevels;
-					InMipTiles.GetKeys(InMipLevels);
-					ENQUEUE_RENDER_COMMAND(CopyFromUploadBuffer)([SampleConverter, ConverterParams, BufferData, BufferRegionsToCopy](FRHICommandListImmediate& RHICmdList)
-					{
-						SCOPED_DRAW_EVENT(RHICmdList, FExrImgMediaReaderGpu_CopyBuffers);
-						TRACE_CPUPROFILER_EVENT_SCOPE_TEXT(*FString::Printf(TEXT("ExrReaderGpu.StartCopy %d"), ConverterParams->FrameId));
-						if (BufferRegionsToCopy.IsEmpty())
-						{
-							RHICmdList.CopyBuffer(BufferData->UploadBufferRef, BufferData->ShaderAccessBufferRef);
-						}
-						else
-						{
-							for (UE::Math::TIntPoint<int64> Region : BufferRegionsToCopy)
-							{
-								RHICmdList.CopyBufferRegion(BufferData->ShaderAccessBufferRef, Region.X, BufferData->UploadBufferRef, Region.X, Region.Y);
-							}
-						}
-					});
-				}
-			}
-			else
-			{
-				UE_LOG(LogImgMedia, Error, TEXT("Could not load %s"), *ImagePath);
-				return false;
-			}
+			ReadResult = ReadMip(CurrentMipLevel, CurrentTileSelection, OutFrame, ConverterParams, SampleConverter, ImagePath, bHasTiles);
 
 			/* Error handling. */
 			if (ReadResult == Success)
@@ -319,7 +358,7 @@ bool FExrImgMediaReaderGpu::ReadFrame(int32 FrameId, const TMap<int32, FImgMedia
 				// To make sure that Media Texture doesn't call the converter if this frame is invalid.
 				OutFrame->SampleConverter.Reset();
 
-				return FExrImgMediaReader::ReadFrame(FrameId, InMipTiles, OutFrame);
+				return FExrImgMediaReader::ReadFrame(ConverterParams->FrameId, InMipTiles, OutFrame);
 			}
 			else if (ReadResult == Cancelled)
 			{
@@ -454,88 +493,123 @@ void FExrImgMediaReaderGpu::CreateSampleConverterCallback
 		SCOPED_DRAW_EVENT(RHICmdList, FExrImgMediaReaderGpu_Convert);
 		SCOPED_GPU_STAT(RHICmdList, ExrImgMediaReaderGpu);
 
-		FIntPoint Dim = ConverterParams->FullResolution;
-		for (TPair<int32, TArray<FIntRect>> MipLevelViewports : ConverterParams->Viewports)
+		auto RenderMip = []
+			( FRHICommandListImmediate& RHICmdList
+			, FTexture2DRHIRef RenderTargetTextureRHI
+			, TSharedPtr<FSampleConverterParameters> ConverterParams
+			, int32 SampleMipLevel
+			, int32 TextureMipLevel
+			, FStructuredBufferPoolItemSharedPtr BufferData
+			, const FIntPoint& SampleSize
+			, const FIntPoint& TextureSize
+			, const TArray<FIntRect>& MipViewports)
 		{
-			int32 MipLevel = MipLevelViewports.Key;
 			SCOPED_GPU_STAT(RHICmdList, ExrImgMediaReaderGpu_MipRender);
-			int MipLevelDiv = 1 << MipLevelViewports.Key;
-			Dim = ConverterParams->FullResolution / MipLevelDiv;
 
-			FStructuredBufferPoolItemSharedPtr BufferData = MipBuffers[MipLevel];
-			if (BufferData.IsValid())
+			// This flag will indicate that we should wait for poll to complete.
+			BufferData->bWillBeSignaled = true;
+
+			FRHIRenderPassInfo RPInfo(RenderTargetTextureRHI, ERenderTargetActions::DontLoad_Store, nullptr, TextureMipLevel);
+			RHICmdList.BeginRenderPass(RPInfo, TEXT("ExrTextureSwizzle"));
+
+			FExrSwizzlePS::FPermutationDomain PermutationVector;
+			PermutationVector.Set<FExrSwizzlePS::FRgbaSwizzle>(ConverterParams->FrameInfo.NumChannels - 1);
+			PermutationVector.Set<FExrSwizzlePS::FRenderTiles>(ConverterParams->FrameInfo.bHasTiles || ConverterParams->bCustomExr);
+			PermutationVector.Set<FExrSwizzlePS::FCustomExr>(ConverterParams->bCustomExr);
+			PermutationVector.Set<FExrSwizzlePS::FPartialTiles>(false);
+
+			FExrSwizzlePS::FParameters Parameters = FExrSwizzlePS::FParameters();
+			Parameters.TextureSize = SampleSize;
+			Parameters.TileSize = ConverterParams->TileDimWithBorders;
+			Parameters.NumChannels = ConverterParams->FrameInfo.NumChannels;
+			if (ConverterParams->FrameInfo.bHasTiles)
 			{
-				if (!BufferData->UploadBufferRef.IsValid() || (CVarExrReaderUseUploadHeap.GetValueOnAnyThread() && !BufferData->ShaderAccessBufferRef->IsValid()))
-				{
-					continue;
-				}
+				Parameters.NumTiles = FIntPoint(FMath::CeilToInt(float(SampleSize.X) / ConverterParams->TileDimWithBorders.X), FMath::CeilToInt(float(SampleSize.Y) / ConverterParams->TileDimWithBorders.Y));
+			}
 
-				// This flag will indicate that we should wait for poll to complete.
-				BufferData->bWillBeSignaled = true;
+			if (ConverterParams->FrameInfo.bHasTiles &&
+				(ConverterParams->TileInfoPerMipLevel.Num() > SampleMipLevel && ConverterParams->TileInfoPerMipLevel[SampleMipLevel].Num() > 0))
+			{
+				TRACE_CPUPROFILER_EVENT_SCOPE_TEXT(*FString::Printf(TEXT("ExrReaderGpu.TileDesc")));
+
+				FBufferRHIRef BufferRef;
+				FRHIResourceCreateInfo CreateInfo(TEXT("FExrImgMediaReaderGpu_TileDesc"));
+				const uint32 BytesPerElement = sizeof(FExrReader::FTileDesc);
+				const uint32 NumElements = ConverterParams->TileInfoPerMipLevel[SampleMipLevel].Num();
+
+				// This buffer is allocated on already allocated block, therefore the risk of fragmentation is mitigated.
+				BufferRef = RHICreateStructuredBuffer(BytesPerElement, BytesPerElement * NumElements, BUF_ShaderResource | BUF_Dynamic | BUF_FastVRAM, CreateInfo);
+				void* MappedBuffer = RHILockBuffer(BufferRef, 0, NumElements, RLM_WriteOnly);
+				FMemory::Memcpy(MappedBuffer, &ConverterParams->TileInfoPerMipLevel[SampleMipLevel][0], sizeof(FExrReader::FTileDesc) * ConverterParams->TileInfoPerMipLevel[SampleMipLevel].Num());
+				RHIUnlockBuffer(BufferRef);
+				Parameters.TileDescBuffer = RHICreateShaderResourceView(BufferRef);
+				PermutationVector.Set<FExrSwizzlePS::FPartialTiles>(true);
+			}
 
-				FRHIRenderPassInfo RPInfo(RenderTargetTextureRHI, ERenderTargetActions::DontLoad_Store, nullptr, MipLevel);
-				RHICmdList.BeginRenderPass(RPInfo, TEXT("ExrTextureSwizzle"));
+			Parameters.UnswizzledBuffer = BufferData->ShaderResourceView;
 
-				FExrSwizzlePS::FPermutationDomain PermutationVector;
-				PermutationVector.Set<FExrSwizzlePS::FRgbaSwizzle>(ConverterParams->FrameInfo.NumChannels - 1);
-				PermutationVector.Set<FExrSwizzlePS::FRenderTiles>(ConverterParams->FrameInfo.bHasTiles || ConverterParams->bCustomExr);
-				PermutationVector.Set<FExrSwizzlePS::FCustomExr>(ConverterParams->bCustomExr);
-				PermutationVector.Set<FExrSwizzlePS::FPartialTiles>(false);
+			FGlobalShaderMap* ShaderMap = GetGlobalShaderMap(GMaxRHIFeatureLevel);
 
-				FExrSwizzlePS::FParameters Parameters = FExrSwizzlePS::FParameters();
-				Parameters.TextureSize = Dim;
-				Parameters.TileSize = ConverterParams->TileDimWithBorders;
-				Parameters.NumChannels = ConverterParams->FrameInfo.NumChannels;
-				if (ConverterParams->FrameInfo.bHasTiles)
-				{
-					Parameters.NumTiles = FIntPoint(FMath::CeilToInt(float(Dim.X) / ConverterParams->TileDimWithBorders.X), FMath::CeilToInt(float(Dim.Y) / ConverterParams->TileDimWithBorders.Y));
-				}
+			TShaderMapRef<FExrSwizzleVS> SwizzleShaderVS(ShaderMap);
+			TShaderMapRef<FExrSwizzlePS> SwizzleShaderPS(ShaderMap, PermutationVector);
 
-				if (ConverterParams->FrameInfo.bHasTiles && 
-					(ConverterParams->TileInfoPerMipLevel.Num() > MipLevel && ConverterParams->TileInfoPerMipLevel[MipLevel].Num() > 0))
-				{
-					TRACE_CPUPROFILER_EVENT_SCOPE_TEXT(*FString::Printf(TEXT("ExrReaderGpu.TileDesc")));
-
-					FBufferRHIRef BufferRef;
-					FRHIResourceCreateInfo CreateInfo(TEXT("FExrImgMediaReaderGpu_TileDesc"));
-					const uint32 BytesPerElement = sizeof(FExrReader::FTileDesc);
-					const uint32 NumElements = ConverterParams->TileInfoPerMipLevel[MipLevel].Num();
-
-					// This buffer is allocated on already allocated block, therefore the risk of fragmentation is mitigated.
-					BufferRef = RHICreateStructuredBuffer(BytesPerElement, BytesPerElement*NumElements, BUF_ShaderResource | BUF_Dynamic | BUF_FastVRAM, CreateInfo);
-					void* MappedBuffer = RHILockBuffer(BufferRef, 0, NumElements, RLM_WriteOnly);
-					FMemory::Memcpy(MappedBuffer, &ConverterParams->TileInfoPerMipLevel[MipLevel][0], sizeof(FExrReader::FTileDesc) * ConverterParams->TileInfoPerMipLevel[MipLevel].Num());
-					RHIUnlockBuffer(BufferRef);
-					Parameters.TileDescBuffer = RHICreateShaderResourceView(BufferRef);
-					PermutationVector.Set<FExrSwizzlePS::FPartialTiles>(true);
-				}
+			FScreenPassPipelineState PipelineState(SwizzleShaderVS, SwizzleShaderPS, TStaticBlendState<>::GetRHI(), TStaticDepthStencilState<false, CF_Always>::GetRHI());
 
-				Parameters.UnswizzledBuffer = BufferData->ShaderResourceView;
+			// If there are tiles determines if we should deliver tiles one by one or in a bulk.
+			for (const FIntRect& Viewport : MipViewports)
+			{
+				DrawScreenPass(RHICmdList, TextureSize, Viewport, PipelineState, [&](FRHICommandListImmediate& RHICmdList)
+					{
+						SetShaderParameters(RHICmdList, SwizzleShaderPS, SwizzleShaderPS.GetPixelShader(), Parameters);
+					});
+			}
 
-				FGlobalShaderMap* ShaderMap = GetGlobalShaderMap(GMaxRHIFeatureLevel);
+			// Resolve render target.
+			RHICmdList.EndRenderPass();
 
-				TShaderMapRef<FExrSwizzleVS> SwizzleShaderVS(ShaderMap);
-				TShaderMapRef<FExrSwizzlePS> SwizzleShaderPS(ShaderMap, PermutationVector);
+			// Mark this render command for this buffer as complete, so we can poll it and transfer later.
+			RHICmdList.WriteGPUFence(BufferData->RenderFence);
+		};
 
-				FScreenPassPipelineState PipelineState(SwizzleShaderVS, SwizzleShaderPS, TStaticBlendState<>::GetRHI(), TStaticDepthStencilState<false, CF_Always>::GetRHI());
-				
-				// If there are tiles determines if we should deliver tiles one by one or in a bulk.
-				for (const FIntRect& Viewport : MipLevelViewports.Value)
-				{
-					DrawScreenPass(RHICmdList, Dim, Viewport, PipelineState, [&](FRHICommandListImmediate& RHICmdList)
-						{
-							SetShaderParameters(RHICmdList, SwizzleShaderPS, SwizzleShaderPS.GetPixelShader(), Parameters);
-						});
-				}
+		int32 MipToUpscale = ConverterParams->UpscaleMip;
 
-				// Resolve render target.
-				RHICmdList.EndRenderPass();
+		// Upscale to all mips below the mip to upscale.
+		for (int32 MipLevel = 0; MipLevel <= MipToUpscale; MipLevel++)
+		{
+			int MipLevelDiv = 1 << MipLevel;
+			FIntPoint Dim = ConverterParams->FullResolution / MipLevelDiv;
+
+			{
+				SCOPED_GPU_STAT(RHICmdList, ExrImgMediaReaderGpu_MipUpscale);
 
-				// Mark this render command for this buffer as complete, so we can poll it and transfer later.
-				RHICmdList.WriteGPUFence(BufferData->RenderFence);
+				FStructuredBufferPoolItemSharedPtr BufferDataToUpscale = MipBuffers[MipToUpscale];
+				TArray<FIntRect> FakeViewport;
+				FakeViewport.Add(FIntRect(FIntPoint(0, 0), Dim));
+				RenderMip(RHICmdList, RenderTargetTextureRHI, ConverterParams, MipToUpscale, MipLevel, BufferDataToUpscale, (ConverterParams->FullResolution / (1 << MipToUpscale)), Dim, FakeViewport);
 			}
+		}
+
+		for (const TPair<int32, TArray<FIntRect>>& MipLevelViewports : ConverterParams->Viewports)
+		{
+			int32 MipLevel = MipLevelViewports.Key;
+			FStructuredBufferPoolItemSharedPtr BufferData = MipBuffers[MipLevel];
+			int MipLevelDiv = 1 << MipLevel;
+			FIntPoint Dim = ConverterParams->FullResolution / MipLevelDiv;
+
+			if (BufferData.IsValid())
+			{
+				if (!BufferData->UploadBufferRef.IsValid() || (bUseUploadHeap && !BufferData->ShaderAccessBufferRef->IsValid()))
+				{
+					continue;
+				}
 
-			
+				// Skip the mip to upscale because it is read and rendered already.
+				if (MipLevelViewports.Key == MipToUpscale)
+				{
+					continue;
+				}
+				RenderMip(RHICmdList, RenderTargetTextureRHI, ConverterParams, MipLevel, MipLevel, BufferData, Dim, Dim, MipLevelViewports.Value);
+			}
 		}
 
 		//Doesn't need further conversion so returning false.
diff --git a/Engine/Plugins/Media/ImgMedia/Source/ImgMedia/Private/Readers/ExrImgMediaReaderGpu.h b/Engine/Plugins/Media/ImgMedia/Source/ImgMedia/Private/Readers/ExrImgMediaReaderGpu.h
--- a/Engine/Plugins/Media/ImgMedia/Source/ImgMedia/Private/Readers/ExrImgMediaReaderGpu.h
+++ b/Engine/Plugins/Media/ImgMedia/Source/ImgMedia/Private/Readers/ExrImgMediaReaderGpu.h
@@ -104,6 +104,18 @@ protected:
 	static void CreateSampleConverterCallback
 		(FExrMediaTextureSampleConverter* SampleConverter, TSharedPtr<FSampleConverterParameters> ConverterParams);
 
+	/**
+	* A function that reads one mip level.
+	*/
+	EReadResult ReadMip
+		( const int32 CurrentMipLevel
+		, const FImgMediaTileSelection& CurrentTileSelection
+		, TSharedPtr<FImgMediaFrame, ESPMode::ThreadSafe> OutFrame
+		, TSharedPtr<FSampleConverterParameters> ConverterParams
+		, FExrMediaTextureSampleConverter* SampleConverter
+		, const FString& ImagePath
+		, bool bHasTiles);
+
 public:
 
 	/** Typically we would need the (ImageResolution.x*y)*NumChannels*ChannelSize */
