ImgMedia: Force upscale requests when the exr cvar is valid.

#rnx

diff --git a/Engine/Plugins/Media/ImgMedia/Source/ImgMedia/Private/Assets/ImgMediaMipMapInfo.cpp b/Engine/Plugins/Media/ImgMedia/Source/ImgMedia/Private/Assets/ImgMediaMipMapInfo.cpp
--- a/Engine/Plugins/Media/ImgMedia/Source/ImgMedia/Private/Assets/ImgMediaMipMapInfo.cpp
+++ b/Engine/Plugins/Media/ImgMedia/Source/ImgMedia/Private/Assets/ImgMediaMipMapInfo.cpp
@@ -806,6 +806,15 @@ TMap<int32, FImgMediaTileSelection> FImgMediaMipMapInfo::GetVisibleTiles()
 	// However the Tick is called from a different thread so care must still be taken when
 	// accessing things that are modified by code external to this function.
 
+
+	int32 MipToUpscale = -1;
+	static const auto CVarUpscaleMip = IConsoleManager::Get().FindTConsoleVariableDataInt(TEXT("r.ExrReaderGPU.UpscaleHigherLevelMip"));
+
+	if (CVarUpscaleMip)
+	{
+		MipToUpscale = FMath::Min(CVarUpscaleMip->GetValueOnAnyThread(), SequenceInfo.NumMipLevels - 1);
+	}
+
 	FScopeLock Lock(&InfoCriticalSection);
 	
 	// Do we need to update the cache?
@@ -815,6 +824,15 @@ TMap<int32, FImgMediaTileSelection> FImgMediaMipMapInfo::GetVisibleTiles()
 
 		CachedVisibleTiles.Reset();
 
+		// We force-add fully visible regions for all levels greater or equal to the upscale mip
+		if (MipToUpscale >= 0)
+		{
+			for (int32 MipLevel = MipToUpscale; MipLevel < SequenceInfo.NumMipLevels; ++MipLevel)
+			{
+				CachedVisibleTiles.Add(MipLevel, FImgMediaTileSelection::CreateForTargetMipLevel(SequenceInfo.Dim, SequenceInfo.TilingDescription.TileSize, MipLevel, true));
+			}
+		}
+
 		FScopeLock LockObjects(&ObjectsCriticalSection);
 
 		// Loop over all objects.
