ImgMedia: Preventative fix for gpu reader deadlock where we prevent locking during reader buffer pre-allocation.

diff --git a/Engine/Plugins/Media/ImgMedia/Source/ImgMedia/Private/Loader/ImgMediaLoader.cpp b/Engine/Plugins/Media/ImgMedia/Source/ImgMedia/Private/Loader/ImgMediaLoader.cpp
--- a/Engine/Plugins/Media/ImgMedia/Source/ImgMedia/Private/Loader/ImgMediaLoader.cpp
+++ b/Engine/Plugins/Media/ImgMedia/Source/ImgMedia/Private/Loader/ImgMediaLoader.cpp
@@ -751,7 +751,13 @@ void FImgMediaLoader::Initialize(const FString& SequencePath, const FFrameRate&
 
 	check(!Initialized); // reinitialization not allowed for now
 
-	LoadSequence(SequencePath, FrameRateOverride, Loop);
+	FImgMediaFrameInfo FirstFrameInfo;
+	
+	if (LoadSequence(SequencePath, FrameRateOverride, Loop, FirstFrameInfo))
+	{
+		WarmupSequence(FirstFrameInfo, Loop);
+	}
+	
 	FPlatformMisc::MemoryBarrier();
 
 	Initialized = true;
@@ -809,13 +815,13 @@ FTimespan FImgMediaLoader::FrameNumberToTime(uint32 FrameNumber) const
 }
 
 
-void FImgMediaLoader::LoadSequence(const FString& SequencePath, const FFrameRate& FrameRateOverride, bool Loop)
+bool FImgMediaLoader::LoadSequence(const FString& SequencePath, const FFrameRate& FrameRateOverride, bool Loop, FImgMediaFrameInfo& FirstFrameInfo)
 {
 	SCOPE_CYCLE_COUNTER(STAT_ImgMedia_LoaderLoadSequence);
 
 	if (SequencePath.IsEmpty())
 	{
-		return;
+		return false;
 	}
 
 	// locate image sequence files
@@ -824,7 +830,7 @@ void FImgMediaLoader::LoadSequence(const FString& SequencePath, const FFrameRate
 	if (FoundPaths.Num() == 0)
 	{
 		UE_LOG(LogImgMedia, Error, TEXT("The directory %s does not contain any image files"), *SequencePath);
-		return;
+		return false;
 	}
 	ImagePaths.Emplace(FoundPaths);
 
@@ -846,7 +852,7 @@ void FImgMediaLoader::LoadSequence(const FString& SequencePath, const FFrameRate
 		Reader = FExrImgMediaReader::GetReader(AsShared(), ImagePaths[0][0]);
 #else
 		UE_LOG(LogImgMedia, Error, TEXT("EXR image sequences are currently supported on macOS and Windows only"));
-		return;
+		return false;
 #endif
 	}
 	else
@@ -856,7 +862,7 @@ void FImgMediaLoader::LoadSequence(const FString& SequencePath, const FFrameRate
 	if (Reader.IsValid() == false)
 	{
 		UE_LOG(LogImgMedia, Error, TEXT("Reader is not valid for file %s."), *ImagePaths[0][0]);
-		return;
+		return false;
 	}
 
 	const UImgMediaSettings* Settings = GetDefault<UImgMediaSettings>();
@@ -869,7 +875,6 @@ void FImgMediaLoader::LoadSequence(const FString& SequencePath, const FFrameRate
 	SequenceName = FName(*SequencePath);
 
 	// fetch sequence attributes from first image
-	FImgMediaFrameInfo FirstFrameInfo;
 	{
 		// Try and get frame from the global cache.
 		const TSharedPtr<FImgMediaFrame, ESPMode::ThreadSafe>* Frame = nullptr;
@@ -885,19 +890,19 @@ void FImgMediaLoader::LoadSequence(const FString& SequencePath, const FFrameRate
 		else if (!Reader->GetFrameInfo(ImagePaths[0][0], FirstFrameInfo))
 		{
 			UE_LOG(LogImgMedia, Error, TEXT("Failed to get frame information from first image in %s"), *SequencePath);
-			return;
+			return false;
 		}
 	}
 	if (FirstFrameInfo.UncompressedSize == 0)
 	{
 		UE_LOG(LogImgMedia, Error, TEXT("The first image in sequence %s does not have a valid frame size"), *SequencePath);
-		return;
+		return false;
 	}
 
 	if (FirstFrameInfo.Dim.GetMin() <= 0)
 	{
 		UE_LOG(LogImgMedia, Error, TEXT("The first image in sequence %s does not have a valid dimension"), *SequencePath);
-		return;
+		return false;
 	}
 
 	NumMipLevels = FMath::Max(FirstFrameInfo.NumMipLevels, ImagePaths.Num());
@@ -980,15 +985,8 @@ void FImgMediaLoader::LoadSequence(const FString& SequencePath, const FFrameRate
 		NumLoadAhead = (int32)((1.0f - LoadBehindScale) * MaxFramesToLoad);
 	}
 
-	// Giving our reader a chance to handle RAM allocation.
-	// Not all readers use this, only those that need to handle large files 
-	// or need to be as efficient as possible.
-	Reader->PreAllocateMemoryPool(NumFramesToLoad, FirstFrameInfo, FirstFrameInfo.FormatName == TEXT("EXR CUSTOM"));
-
 	Frames.Empty(NumFramesToLoad);
 
-	Update(0, 0.0f, Loop);
-
 	// update info
 	Info = TEXT("Image Sequence\n");
 	Info += FString::Printf(TEXT("    Dimension: %i x %i\n"), SequenceDim.X, SequenceDim.Y);
@@ -996,8 +994,21 @@ void FImgMediaLoader::LoadSequence(const FString& SequencePath, const FFrameRate
 	Info += FString::Printf(TEXT("    Compression: %s\n"), *FirstFrameInfo.CompressionName);
 	Info += FString::Printf(TEXT("    Frames: %i\n"), GetNumImages());
 	Info += FString::Printf(TEXT("    Frame Rate: %.2f (%i/%i)\n"), SequenceFrameRate.AsDecimal(), SequenceFrameRate.Numerator, SequenceFrameRate.Denominator);
+
+	return true;
 }
 
+void FImgMediaLoader::WarmupSequence(const FImgMediaFrameInfo& InFirstFrameInfo, bool Loop)
+{
+	// Giving our reader a chance to handle RAM allocation.
+	// Not all readers use this, only those that need to handle large files 
+	// or need to be as efficient as possible.
+	Reader->PreAllocateMemoryPool(NumFramesToLoad, InFirstFrameInfo, InFirstFrameInfo.FormatName == TEXT("EXR CUSTOM"));
+
+	FScopeLock Lock(&CriticalSection);
+
+	Update(0, 0.0f, Loop);
+}
 
 void FImgMediaLoader::FindFiles(const FString& SequencePath, TArray<FString>& OutputPaths)
 {
diff --git a/Engine/Plugins/Media/ImgMedia/Source/ImgMedia/Private/Loader/ImgMediaLoader.h b/Engine/Plugins/Media/ImgMedia/Source/ImgMedia/Private/Loader/ImgMediaLoader.h
--- a/Engine/Plugins/Media/ImgMedia/Source/ImgMedia/Private/Loader/ImgMediaLoader.h
+++ b/Engine/Plugins/Media/ImgMedia/Source/ImgMedia/Private/Loader/ImgMediaLoader.h
@@ -22,6 +22,7 @@ class IImgMediaReader;
 class IQueuedWork;
 
 struct FImgMediaFrame;
+struct FImgMediaFrameInfo;
 
 /**
  * Settings for the smart cache.
@@ -384,8 +385,18 @@ protected:
 	 * @param SequencePath Path to the image sequence.
 	 * @param FrameRateOverride The frame rate to use (0/0 = do not override).
 	 * @param Loop Whether the cache should loop around.
+	 * @param OutFirstFrameInfo Sequence information based on its first frame.
+	 * @return True when loading succeeds.
 	 */
-	void LoadSequence(const FString& SequencePath, const FFrameRate& FrameRateOverride, bool Loop);
+	bool LoadSequence(const FString& SequencePath, const FFrameRate& FrameRateOverride, bool Loop, FImgMediaFrameInfo& OutFirstFrameInfo);
+
+	/**
+	 * Warms up the player and start issuing requests from the start of the sequence.
+	 *
+	 * @param InFirstFrameInfo Sequence information based on its first frame.
+	 * @param Loop Whether the cache should loop around.
+	 */
+	void WarmupSequence(const FImgMediaFrameInfo& InFirstFrameInfo, bool Loop);
 
 	/**
 	 * Finds all the files in a directory and gets their path.
