ImgMedia: Since tile size is now used in calculations, prevent using invalid/zero sizes on sequences with mips but no tiles.

diff --git a/Engine/Plugins/Media/ImgMedia/Source/ImgMedia/Private/Assets/ImgMediaMipMapInfo.cpp b/Engine/Plugins/Media/ImgMedia/Source/ImgMedia/Private/Assets/ImgMediaMipMapInfo.cpp
--- a/Engine/Plugins/Media/ImgMedia/Source/ImgMedia/Private/Assets/ImgMediaMipMapInfo.cpp
+++ b/Engine/Plugins/Media/ImgMedia/Source/ImgMedia/Private/Assets/ImgMediaMipMapInfo.cpp
@@ -41,6 +41,11 @@ FImgMediaTileSelection FImgMediaTileSelection::CreateForTargetMipLevel(const FIn
 {
 	ensure(TargetMipLevel >= 0);
 
+	if (TileSize.X == 0 || TileSize.Y == 0)
+	{
+		return FImgMediaTileSelection(1, 1, bDefaultVisibility);
+	}
+
 	const int MipLevelDiv = 1 << TargetMipLevel;
 	int32 NumTilesX = FMath::Max(1, FMath::CeilToInt(float(MipZeroResolution.X / MipLevelDiv) / TileSize.X));
 	int32 NumTilesY = FMath::Max(1, FMath::CeilToInt(float(MipZeroResolution.Y / MipLevelDiv) / TileSize.Y));
@@ -372,9 +377,13 @@ namespace {
 				const int32 MaxLevel = InSequenceInfo.NumMipLevels - 1;
 				int MipLevelDiv = 1 << MaxLevel;
 
-				FIntPoint CurrentNumTiles;
-				CurrentNumTiles.X = FMath::CeilToInt(float(InSequenceInfo.Dim.X / MipLevelDiv) / InSequenceInfo.TilingDescription.TileSize.X);
-				CurrentNumTiles.Y = FMath::CeilToInt(float(InSequenceInfo.Dim.Y / MipLevelDiv) / InSequenceInfo.TilingDescription.TileSize.Y);
+				FIntPoint CurrentNumTiles = FIntPoint(1, 1);
+
+				if (InSequenceInfo.IsTiled())
+				{
+					CurrentNumTiles.X = FMath::CeilToInt(float(InSequenceInfo.Dim.X / MipLevelDiv) / InSequenceInfo.TilingDescription.TileSize.X);
+					CurrentNumTiles.Y = FMath::CeilToInt(float(InSequenceInfo.Dim.Y / MipLevelDiv) / InSequenceInfo.TilingDescription.TileSize.Y);
+				}
 
 				// Starting with tiles at the highest mip level
 				TQueue<FIntVector> Tiles;
@@ -395,9 +404,12 @@ namespace {
 					int32 CurrentMipLevel = Tile.Z;
 					MipLevelDiv = 1 << CurrentMipLevel;
 
-					// Calculate the number of tiles at this mip level
-					CurrentNumTiles.X = FMath::Max(1, FMath::CeilToInt((float(InSequenceInfo.Dim.X) / MipLevelDiv) / InSequenceInfo.TilingDescription.TileSize.X));
-					CurrentNumTiles.Y = FMath::Max(1, FMath::CeilToInt((float(InSequenceInfo.Dim.Y) / MipLevelDiv) / InSequenceInfo.TilingDescription.TileSize.Y));
+					if (InSequenceInfo.IsTiled())
+					{
+						// Calculate the number of tiles at this mip level
+						CurrentNumTiles.X = FMath::Max(1, FMath::CeilToInt((float(InSequenceInfo.Dim.X) / MipLevelDiv) / InSequenceInfo.TilingDescription.TileSize.X));
+						CurrentNumTiles.Y = FMath::Max(1, FMath::CeilToInt((float(InSequenceInfo.Dim.Y) / MipLevelDiv) / InSequenceInfo.TilingDescription.TileSize.Y));
+					}
 
 					FVector2f CurrentPartialTileNum = InSequenceInfo.GetPartialTileNum(CurrentMipLevel);
 
