Exr Reader GPU: Fixed an assert caused by command list flush in CopyBuffer function.

diff --git a/Engine/Plugins/Media/ImgMedia/Source/ImgMedia/Private/Readers/ExrImgMediaReaderGpu.cpp b/Engine/Plugins/Media/ImgMedia/Source/ImgMedia/Private/Readers/ExrImgMediaReaderGpu.cpp
--- a/Engine/Plugins/Media/ImgMedia/Source/ImgMedia/Private/Readers/ExrImgMediaReaderGpu.cpp
+++ b/Engine/Plugins/Media/ImgMedia/Source/ImgMedia/Private/Readers/ExrImgMediaReaderGpu.cpp
@@ -216,7 +216,7 @@ FExrImgMediaReader::EReadResult FExrImgMediaReaderGpu::ReadMip
 					TRACE_CPUPROFILER_EVENT_SCOPE_TEXT(*FString::Printf(TEXT("ExrReaderGpu.StartCopy %d"), ConverterParams->FrameId));
 					if (BufferRegionsToCopy.IsEmpty())
 					{
-						RHICmdList.CopyBuffer(BufferData->UploadBufferRef, BufferData->ShaderAccessBufferRef);
+						RHICmdList.CopyBufferRegion(BufferData->ShaderAccessBufferRef, 0, BufferData->UploadBufferRef, 0, BufferData->ShaderAccessBufferRef->GetSize());
 					}
 					else
 					{
