ImgMedia: Fixed work not being scheduled when a sample was returned.
#rnx

diff --git a/Engine/Plugins/Media/ImgMedia/Source/ImgMedia/Private/Player/ImgMediaPlayer.cpp b/Engine/Plugins/Media/ImgMedia/Source/ImgMedia/Private/Player/ImgMediaPlayer.cpp
--- a/Engine/Plugins/Media/ImgMedia/Source/ImgMedia/Private/Player/ImgMediaPlayer.cpp
+++ b/Engine/Plugins/Media/ImgMedia/Source/ImgMedia/Private/Player/ImgMediaPlayer.cpp
@@ -793,6 +793,7 @@ IMediaSamples::EFetchBestSampleResult FImgMediaPlayer::FetchBestVideoSampleForTi
 	{
 		// See if we have any samples in the specified time range.
 		SampleResult = Loader->FetchBestVideoSampleForTimeRange(TimeRange, OutSample, ShouldLoop, CurrentRate, PlaybackIsBlocking);
+		Scheduler->TickInput(FTimespan::Zero(), FTimespan::MinValue());
 		if (SampleResult == IMediaSamples::EFetchBestSampleResult::Ok)
 		{
 			CurrentTime = OutSample->GetTime().Time;
