ImgMedia: ExrReader GPU: Fixed an issue when two threads access the same sample converter.

diff --git a/Engine/Plugins/Media/ImgMedia/Source/ImgMedia/Private/Readers/ExrImgMediaReaderGpu.cpp b/Engine/Plugins/Media/ImgMedia/Source/ImgMedia/Private/Readers/ExrImgMediaReaderGpu.cpp
--- a/Engine/Plugins/Media/ImgMedia/Source/ImgMedia/Private/Readers/ExrImgMediaReaderGpu.cpp
+++ b/Engine/Plugins/Media/ImgMedia/Source/ImgMedia/Private/Readers/ExrImgMediaReaderGpu.cpp
@@ -144,6 +144,8 @@ FExrImgMediaReader::EReadResult FExrImgMediaReaderGpu::ReadMip
 	FIntPoint CurrentMipDim = ConverterParams->FullResolution / MipLevelDiv;
 
 	const SIZE_T BufferSize = GetBufferSize(CurrentMipDim, ConverterParams->FrameInfo.NumChannels, bHasTiles, OutFrame->Info.NumTiles / MipLevelDiv, ConverterParams->bCustomExr);
+	
+	SampleConverter->LockMipBuffers();
 	FStructuredBufferPoolItemSharedPtr BufferData = SampleConverter->GetMipLevelBuffer(CurrentMipLevel);
 
 	if (!BufferData.IsValid())
@@ -151,6 +153,7 @@ FExrImgMediaReader::EReadResult FExrImgMediaReaderGpu::ReadMip
 		BufferData = AllocateGpuBufferFromPool(BufferSize);
 		SampleConverter->SetMipLevelBuffer(CurrentMipLevel, BufferData);
 	}
+	SampleConverter->UnlockMipBuffers();
 
 	uint16* MipDataPtr = static_cast<uint16*>(BufferData->UploadBufferMapped);
 
@@ -756,7 +759,12 @@ bool FExrMediaTextureSampleConverter::Convert(FTexture2DRHIRef& InDstTexture, co
 	bool bExecutionSuccessful = false;
 	if (ConvertExrBufferCallback.IsBound())
 	{
-		bExecutionSuccessful = ConvertExrBufferCallback.Execute(FRHICommandListExecutor::GetImmediateCommandList(), InDstTexture, MipBuffers);
+		TMap<int32, FStructuredBufferPoolItemSharedPtr> MipBufferCopy;
+		{
+			FScopeLock MipScopeLock(&MipBufferCriticalSection);
+			MipBufferCopy = MipBuffers;
+		}
+		bExecutionSuccessful = ConvertExrBufferCallback.Execute(FRHICommandListExecutor::GetImmediateCommandList(), InDstTexture, MipBufferCopy);
 	}
 	return bExecutionSuccessful;
 }
diff --git a/Engine/Plugins/Media/ImgMedia/Source/ImgMedia/Private/Readers/ExrImgMediaReaderGpu.h b/Engine/Plugins/Media/ImgMedia/Source/ImgMedia/Private/Readers/ExrImgMediaReaderGpu.h
--- a/Engine/Plugins/Media/ImgMedia/Source/ImgMedia/Private/Readers/ExrImgMediaReaderGpu.h
+++ b/Engine/Plugins/Media/ImgMedia/Source/ImgMedia/Private/Readers/ExrImgMediaReaderGpu.h
@@ -166,25 +166,35 @@ public:
 		ConvertExrBufferCallback = Callback;
 	};
 
+	void LockMipBuffers()
+	{
+		MipBufferCriticalSection.Lock();
+	}
+
+	void UnlockMipBuffers()
+	{
+		MipBufferCriticalSection.Unlock();
+	}
+
 	FStructuredBufferPoolItemSharedPtr GetMipLevelBuffer(int32 RequestedMipLevel)
 	{
-		FScopeLock ScopeLock(&ConverterCallbacksCriticalSection);
 		if (MipBuffers.Contains(RequestedMipLevel))
 		{
 			return *MipBuffers.Find(RequestedMipLevel);
 		}
+
 		return nullptr;
 	}
 
 	void SetMipLevelBuffer(int32 RequestedMipLevel, FStructuredBufferPoolItemSharedPtr Buffer)
 	{
-		FScopeLock ScopeLock(&ConverterCallbacksCriticalSection);
 		check(!MipBuffers.Contains(RequestedMipLevel));
 		MipBuffers.Add(RequestedMipLevel, Buffer);
 	}
 
 private:
 	FCriticalSection ConverterCallbacksCriticalSection;
+	FCriticalSection MipBufferCriticalSection;
 	FExrConvertBufferCallback ConvertExrBufferCallback;
 
 	/** An array of structured buffers that are big enough to fully contain corresponding mip levels. */
