ImgMedia: Give users control over the maximum number of times a frame can be invalidated.

This can become useful when upsampling is enabled, or missing tiles matter less than perf/frame drops.

diff --git a/Engine/Plugins/Media/ImgMedia/Source/ImgMedia/Private/Loader/ImgMediaLoader.cpp b/Engine/Plugins/Media/ImgMedia/Source/ImgMedia/Private/Loader/ImgMediaLoader.cpp
--- a/Engine/Plugins/Media/ImgMedia/Source/ImgMedia/Private/Loader/ImgMediaLoader.cpp
+++ b/Engine/Plugins/Media/ImgMedia/Source/ImgMedia/Private/Loader/ImgMediaLoader.cpp
@@ -37,6 +37,11 @@ DECLARE_CYCLE_STAT(TEXT("ImgMedia Loader Release Cache"), STAT_ImgMedia_LoaderRe
 
 constexpr int32 FImgMediaLoader::MAX_MIPMAP_LEVELS;
 
+static TAutoConsoleVariable<int32> CVarImgMediaFrameInvalidationMaxCount(
+	TEXT("ImgMedia.FrameInvalidationMaxCount"),
+	2,
+	TEXT("Maximum number of cached frames that can be invalidated when missing the latest mips/tiles."));
+
 namespace ImgMediaLoader
 {
 	void CheckAndUpdateImgDimensions(FIntPoint& InOutSequenceDim, const FIntPoint& InNewDim)
@@ -83,6 +88,31 @@ namespace ImgMediaLoader
 		// Requested tiles already exist, or the request was empty.
 		return true;
 	}
+
+	// Check if [CurrentFrame] is contained in the [OffsetCount] number of frames following [OriginFrame], also taking into account looping and play direction.
+	bool IsFrameInRange(int32 CurrentFrame, int32 OriginFrame, int32 OffsetCount, int32 TotalNumFrames, float PlayRate)
+	{
+		if (OffsetCount <= 0)
+		{
+			return false;
+		}
+
+		ensure(CurrentFrame >= 0 && CurrentFrame < TotalNumFrames);
+		ensure(OriginFrame >= 0 && OriginFrame < TotalNumFrames);
+
+		const int32 PlayRateSign = (PlayRate >= 0.0f) ? 1 : -1;
+
+		for (int32 Offset = 0; Offset < OffsetCount; ++Offset)
+		{
+			int32 OffsetFrame = FMath::Wrap(OriginFrame + PlayRateSign * Offset, 0, TotalNumFrames - 1);
+			if (OffsetFrame == CurrentFrame)
+			{
+				return true;
+			}
+		}
+
+		return false;
+	}
 }
 
 FImgMediaLoaderBandwidth::FImgMediaLoaderBandwidth()
@@ -1182,12 +1212,11 @@ void FImgMediaLoader::Update(int32 PlayHeadFrame, float PlayRate, bool Loop)
 	Reader->OnTick();
 
 	// @todo gmp: ImgMedia: take PlayRate and DeltaTime into account when determining frames to load
-	
+	const int32 NumImagePaths = GetNumImages();
+
 	// determine frame numbers to be loaded
 	TArray<int32> FramesToLoad;
 	{
-		const int32 NumImagePaths = GetNumImages();
-
 		FramesToLoad.Empty(NumFramesToLoad);
 
 		int32 FrameOffset = (PlayRate >= 0.0f) ? 1 : -1;
@@ -1291,6 +1320,9 @@ void FImgMediaLoader::Update(int32 PlayHeadFrame, float PlayRate, bool Loop)
 	// determine frame numbers that need to be cached
 	PendingFrameNumbers.Empty();
 
+	// Limit which frames can be invalidated for missing mips/tiles
+	const int32 FrameInvalidationMaxCount = CVarImgMediaFrameInvalidationMaxCount.GetValueOnAnyThread();
+
 	for (int32 FrameNumber : FramesToLoad)
 	{
 		// Get frame from cache.
@@ -1311,13 +1343,18 @@ void FImgMediaLoader::Update(int32 PlayHeadFrame, float PlayRate, bool Loop)
 			// No, we need one.
 			NeedFrame = true;
 		}
-		else
+		else if(ImgMediaLoader::IsFrameInRange(FrameNumber, PlayHeadFrame, FrameInvalidationMaxCount, NumImagePaths, PlayRate))
 		{
 			// Yes. Check if we have all the desired tiles per mip level.
 			TMap<int32, FImgMediaTileSelection> DesiredMipsAndTiles;
 			GetDesiredMipTiles(FrameNumber, DesiredMipsAndTiles);
 
 			NeedFrame = !ImgMediaLoader::ContainsMipTiles(*FramePtr, DesiredMipsAndTiles);
+
+			if (NeedFrame)
+			{
+				UE_LOG(LogImgMedia, VeryVerbose, TEXT("Loader %p: Invalidating frame %i at playhead %i due to missing mips or tiles."), this, FrameNumber, PlayHeadFrame);
+			}
 		}
 		
 		if ((NeedFrame) && !QueuedFrameNumbers.Contains(FrameNumber))
