ImgMedia: Allocator deadlock fix.

diff --git a/Engine/Plugins/Media/ImgMedia/Source/ImgMedia/Private/Readers/ExrImgMediaReaderGpu.cpp b/Engine/Plugins/Media/ImgMedia/Source/ImgMedia/Private/Readers/ExrImgMediaReaderGpu.cpp
--- a/Engine/Plugins/Media/ImgMedia/Source/ImgMedia/Private/Readers/ExrImgMediaReaderGpu.cpp
+++ b/Engine/Plugins/Media/ImgMedia/Source/ImgMedia/Private/Readers/ExrImgMediaReaderGpu.cpp
@@ -759,12 +759,7 @@ bool FExrMediaTextureSampleConverter::Convert(FTexture2DRHIRef& InDstTexture, co
 	bool bExecutionSuccessful = false;
 	if (ConvertExrBufferCallback.IsBound())
 	{
-		TMap<int32, FStructuredBufferPoolItemSharedPtr> MipBufferCopy;
-		{
-			FScopeLock MipScopeLock(&MipBufferCriticalSection);
-			MipBufferCopy = MipBuffers;
-		}
-		bExecutionSuccessful = ConvertExrBufferCallback.Execute(FRHICommandListExecutor::GetImmediateCommandList(), InDstTexture, MipBufferCopy);
+		bExecutionSuccessful = ConvertExrBufferCallback.Execute(FRHICommandListExecutor::GetImmediateCommandList(), InDstTexture, MipBuffersRenderThread);
 	}
 	return bExecutionSuccessful;
 }
diff --git a/Engine/Plugins/Media/ImgMedia/Source/ImgMedia/Private/Readers/ExrImgMediaReaderGpu.h b/Engine/Plugins/Media/ImgMedia/Source/ImgMedia/Private/Readers/ExrImgMediaReaderGpu.h
--- a/Engine/Plugins/Media/ImgMedia/Source/ImgMedia/Private/Readers/ExrImgMediaReaderGpu.h
+++ b/Engine/Plugins/Media/ImgMedia/Source/ImgMedia/Private/Readers/ExrImgMediaReaderGpu.h
@@ -164,6 +164,9 @@ public:
 	{
 		FScopeLock ScopeLock(&ConverterCallbacksCriticalSection);
 		ConvertExrBufferCallback = Callback;
+
+		// Copy mip buffers to be used for rendering. 
+		MipBuffersRenderThread = MipBuffers;
 	};
 
 	void LockMipBuffers()
@@ -194,11 +197,16 @@ public:
 
 private:
 	FCriticalSection ConverterCallbacksCriticalSection;
-	FCriticalSection MipBufferCriticalSection;
 	FExrConvertBufferCallback ConvertExrBufferCallback;
 
-	/** An array of structured buffers that are big enough to fully contain corresponding mip levels. */
+	/** Lock to be used exclusively on reader threads.*/
+	FCriticalSection MipBufferCriticalSection;
+
+	/** An array of structured buffers that are big enough to fully contain corresponding mip levels used by reader threads and transferred into MipBuffersRenderThread. */
 	TMap<int32,FStructuredBufferPoolItemSharedPtr> MipBuffers;
+
+	/** An array of structured buffers that are big enough to fully contain corresponding mip levels. Used exclusively on the Render thread. */
+	TMap<int32, FStructuredBufferPoolItemSharedPtr> MipBuffersRenderThread;
 };
 
 #endif //defined(PLATFORM_WINDOWS) && PLATFORM_WINDOWS
