ImgMedia: Process exr button is disabled if no output path is set.
#rnx

diff --git a/Engine/Plugins/Media/ImgMedia/Source/ImgMediaEditor/Private/Widgets/SImgMediaProcessEXR.cpp b/Engine/Plugins/Media/ImgMedia/Source/ImgMediaEditor/Private/Widgets/SImgMediaProcessEXR.cpp
--- a/Engine/Plugins/Media/ImgMedia/Source/ImgMediaEditor/Private/Widgets/SImgMediaProcessEXR.cpp
+++ b/Engine/Plugins/Media/ImgMedia/Source/ImgMediaEditor/Private/Widgets/SImgMediaProcessEXR.cpp
@@ -74,6 +74,11 @@ void SImgMediaProcessEXR::Construct(const FArguments& InArgs)
 					[
 						SAssignNew(StartButton, SButton)
 							.ButtonStyle(&FAppStyle::Get().GetWidgetStyle<FButtonStyle>("PrimaryButton"))
+							.IsEnabled_Lambda([this]
+							{
+								return (bIsProcessing == false) && (Options.IsValid()) &&
+									(Options->OutputPath.Path.IsEmpty() == false);
+							})
 							.OnClicked(this, &SImgMediaProcessEXR::OnProcessImagesClicked)
 							.ForegroundColor(FSlateColor::UseStyle())
 							.Text(LOCTEXT("StartProcessImages", "Process Images"))
@@ -131,7 +136,6 @@ void SImgMediaProcessEXR::Tick(const FGeometry& AllottedGeometry, const double I
 
 void SImgMediaProcessEXR::UpdateWidgets()
 {
-	StartButton->SetEnabled(!bIsProcessing);
 	CancelButton->SetEnabled(bIsProcessing && (!bIsCancelling));
 }
 
