ImgMedia: Fix EXR gpu reader issue on cached frame updates.

If the list of visible tiles changed on one mip level but not others, viewports for the levels with idendical set of tiles would be eliminated causing stalling mips in the media texture.

#rnx

diff --git a/Engine/Plugins/Media/ImgMedia/Source/ImgMedia/Private/Readers/ExrImgMediaReader.h b/Engine/Plugins/Media/ImgMedia/Source/ImgMedia/Private/Readers/ExrImgMediaReader.h
--- a/Engine/Plugins/Media/ImgMedia/Source/ImgMedia/Private/Readers/ExrImgMediaReader.h
+++ b/Engine/Plugins/Media/ImgMedia/Source/ImgMedia/Private/Readers/ExrImgMediaReader.h
@@ -46,7 +46,8 @@ protected:
 	{
 		Fail,
 		Success,
-		Cancelled
+		Cancelled,
+		Skipped
 	};
 
 	/* 
diff --git a/Engine/Plugins/Media/ImgMedia/Source/ImgMedia/Private/Readers/ExrImgMediaReaderGpu.cpp b/Engine/Plugins/Media/ImgMedia/Source/ImgMedia/Private/Readers/ExrImgMediaReaderGpu.cpp
--- a/Engine/Plugins/Media/ImgMedia/Source/ImgMedia/Private/Readers/ExrImgMediaReaderGpu.cpp
+++ b/Engine/Plugins/Media/ImgMedia/Source/ImgMedia/Private/Readers/ExrImgMediaReaderGpu.cpp
@@ -191,11 +191,19 @@ FExrImgMediaReader::EReadResult FExrImgMediaReaderGpu::ReadMip
 				Viewports.Add(MoveTemp(Viewport));
 			}
 
-			ReadResult = ReadTiles(MipDataPtr, BufferSize, ImagePath, ConverterParams->FrameId, TileRegionsToRead, ConverterParams, CurrentMipLevel, BufferRegionsToCopy);
-
-			for (const FIntRect& Region : TileRegionsToRead)
+			if (TileRegionsToRead.IsEmpty() && !TileRegionsToRender.IsEmpty())
 			{
-				OutFrame->NumTilesRead += Region.Area();
+				// If all tiles were previously read and stored in cached frame, reading can be skipped.
+				ReadResult = Skipped;
+			}
+			else
+			{
+				ReadResult = ReadTiles(MipDataPtr, BufferSize, ImagePath, ConverterParams->FrameId, TileRegionsToRead, ConverterParams, CurrentMipLevel, BufferRegionsToCopy);
+				
+				for (const FIntRect& Region : TileRegionsToRead)
+				{
+					OutFrame->NumTilesRead += Region.Area();
+				}
 			}
 		}
 		else
@@ -319,23 +327,15 @@ bool FExrImgMediaReaderGpu::ReadFrame(int32 FrameId, const TMap<int32, FImgMedia
 			const int32 CurrentMipLevel = TilesPerMip.Key;
 			const FImgMediaTileSelection& CurrentTileSelection = TilesPerMip.Value;
 
-			bool ReadThisMip = true;
-			// Avoid reads if the cached frame already contains the current tiles for this mip level.
-			if (const FImgMediaTileSelection* CachedSelection = OutFrame->MipTilesPresent.Find(CurrentMipLevel))
-			{
-				ReadThisMip = !CachedSelection->Contains(CurrentTileSelection);
-			}
-
-			if (!ReadThisMip)
+			if (!CurrentTileSelection.IsAnyVisible())
 			{
 				continue;
 			}
-			EReadResult ReadResult = Fail;
 
 			// Get highest resolution mip level path.
 			FString ImagePath = Loader->GetImagePath(ConverterParams->FrameId, ConverterParams->bMipsInSeparateFiles ? CurrentMipLevel : 0);
 
-			ReadResult = ReadMip(CurrentMipLevel, CurrentTileSelection, OutFrame, ConverterParams, SampleConverter, ImagePath, bHasTiles);
+			EReadResult ReadResult = ReadMip(CurrentMipLevel, CurrentTileSelection, OutFrame, ConverterParams, SampleConverter, ImagePath, bHasTiles);
 
 			/* Error handling. */
 			if (ReadResult == Success)
