ImgMedia: Added mip level estimation padding control.

#rnx

diff --git a/Engine/Plugins/Media/ImgMedia/Source/ImgMedia/Private/Assets/ImgMediaMipMapInfo.cpp b/Engine/Plugins/Media/ImgMedia/Source/ImgMedia/Private/Assets/ImgMediaMipMapInfo.cpp
--- a/Engine/Plugins/Media/ImgMedia/Source/ImgMedia/Private/Assets/ImgMediaMipMapInfo.cpp
+++ b/Engine/Plugins/Media/ImgMedia/Source/ImgMedia/Private/Assets/ImgMediaMipMapInfo.cpp
@@ -24,11 +24,18 @@ DECLARE_CYCLE_STAT(TEXT("ImgMedia MipMap Update Cache"), STAT_ImgMedia_MipMapUpd
 static TAutoConsoleVariable<bool> CVarImgMediaMipMapDebugEnable(
 	TEXT("ImgMedia.MipMapDebug"),
 	0,
-	TEXT("Display debug on mipmaps used by the ImgMedia plugin.\n")
+	TEXT("Display debug on mipmaps and tiles used by the ImgMedia plugin.\n")
 	TEXT("   0: off (default)\n")
 	TEXT("   1: on\n"),
 	ECVF_Default);
 
+
+static TAutoConsoleVariable<float> CVarImgMediaMipLevelPadding(
+	TEXT("ImgMedia.MipMapLevelPadding"),
+	0.0,
+	TEXT("Value padded onto the estimated (minimum and maximum) mipmap levels used by the loader.\n"),
+	ECVF_Default);
+
 FImgMediaTileSelection::FImgMediaTileSelection(int32 NumTilesX, int32 NumTilesY, bool bDefaultVisibility)
 	: Tiles(bDefaultVisibility, NumTilesX * NumTilesY)
 	, Dimensions(NumTilesX, NumTilesY)
@@ -346,6 +353,7 @@ namespace {
 			{
 				return;
 			}
+			const float MipMapLevelPadding = FMath::Max(CVarImgMediaMipLevelPadding.GetValueOnAnyThread(), 0.0f);
 
 			const FIntPoint& SequenceTileNum = InSequenceInfo.TilingDescription.TileNum;
 
@@ -468,8 +476,8 @@ namespace {
 								
 								if (bValidLevel)
 								{
-									MipLevelRange[0] = FMath::Min(MipLevelRange[0], FMath::Clamp((int32)CalculatedLevel, 0, MaxLevel));
-									MipLevelRange[1] = FMath::Max(MipLevelRange[1], FMath::Clamp(FMath::CeilToInt32(CalculatedLevel), 0, MaxLevel));
+									MipLevelRange[0] = FMath::Min(MipLevelRange[0], FMath::Clamp(FMath::FloorToInt32(CalculatedLevel - MipMapLevelPadding), 0, MaxLevel));
+									MipLevelRange[1] = FMath::Max(MipLevelRange[1], FMath::Clamp(FMath::CeilToInt32(CalculatedLevel + MipMapLevelPadding), 0, MaxLevel));
 									NumVisibleCorners++;
 								}
 							}
@@ -575,6 +583,8 @@ namespace {
 				return;
 			}
 
+			const float MipMapLevelPadding = FMath::Max(CVarImgMediaMipLevelPadding.GetValueOnAnyThread(), 0.0f);
+
 			const FTransform MeshTransform = Mesh->GetComponentTransform();
 			const float DefaultSphereRadius = 50.0f;
 			const int32 MaxLevel = InSequenceInfo.NumMipLevels - 1;
@@ -630,8 +640,8 @@ namespace {
 							{
 								CalculatedLevel += MipMapBias + ViewInfo.MaterialTextureMipBias;
 
-								MipLevelRange[0] = FMath::Clamp((int32)CalculatedLevel, 0, MaxLevel);
-								MipLevelRange[1] = FMath::CeilToInt32(CalculatedLevel);
+								MipLevelRange[0] = FMath::Clamp(FMath::FloorToInt32(CalculatedLevel - MipMapLevelPadding), 0, MaxLevel);
+								MipLevelRange[1] = FMath::CeilToInt32(CalculatedLevel + MipMapLevelPadding);
 
 								// As a mitigation for discontinuities at the poles, we artifically increase the max calculated level.
 								// (Note: Using an icosphere would avoid this issue but conflict with the partial sphere feature.)
