ImgMedia: Fixed crash when trying to process a non valid exr file.
#rnx

diff --git a/Engine/Plugins/Media/ImgMedia/Source/OpenExrWrapper/Private/OpenExrWrapper.cpp b/Engine/Plugins/Media/ImgMedia/Source/OpenExrWrapper/Private/OpenExrWrapper.cpp
--- a/Engine/Plugins/Media/ImgMedia/Source/OpenExrWrapper/Private/OpenExrWrapper.cpp
+++ b/Engine/Plugins/Media/ImgMedia/Source/OpenExrWrapper/Private/OpenExrWrapper.cpp
@@ -144,6 +144,11 @@ FFrameRate FRgbaInputFile::GetFrameRate(const FFrameRate& DefaultValue) const
 
 int32 FRgbaInputFile::GetNumChannels() const
 {
+	if (InputFile == nullptr)
+	{
+		return 0;
+	}
+
 	Imf::RgbaChannels Channels = ((Imf::RgbaInputFile*)InputFile)->channels();
 	int32 NumChannels = 3;
 	switch (Channels)
