ImgMedia: Fix RHI fence use for latest parallel rendering API changes.

#rnx

diff --git a/Engine/Plugins/Media/ImgMedia/Source/ImgMedia/Private/Readers/ExrImgMediaReaderGpu.cpp b/Engine/Plugins/Media/ImgMedia/Source/ImgMedia/Private/Readers/ExrImgMediaReaderGpu.cpp
--- a/Engine/Plugins/Media/ImgMedia/Source/ImgMedia/Private/Readers/ExrImgMediaReaderGpu.cpp
+++ b/Engine/Plugins/Media/ImgMedia/Source/ImgMedia/Private/Readers/ExrImgMediaReaderGpu.cpp
@@ -737,10 +737,11 @@ void FExrImgMediaReaderGpu::TransferFromStagingBuffer()
 			for (FStructuredBufferPoolItem* MemoryPoolItem : AllValues)
 			{
 				// Check if fence has signaled. Or otherwise if we are waiting for signal to come through.
-				if (MemoryPoolItem->RenderFence->Poll() || !MemoryPoolItem->bWillBeSignaled)
+				if (!MemoryPoolItem->bWillBeSignaled || MemoryPoolItem->RenderFence->Poll())
 				{
 					// If buffer was in use but the fence signaled we need to reset bWillBeSignaled flag.
 					MemoryPoolItem->bWillBeSignaled = false;
+					MemoryPoolItem->RenderFence->Clear();
 
 					StagingMemoryPool.Remove(Key, MemoryPoolItem);
 					MemoryPool.Add(Key, MemoryPoolItem);
