ImgMedia: Case missed in RHI update 22229235, where upscaling or stopped playback would set the fence more than once.

#rnx

diff --git a/Engine/Plugins/Media/ImgMedia/Source/ImgMedia/Private/Readers/ExrImgMediaReaderGpu.cpp b/Engine/Plugins/Media/ImgMedia/Source/ImgMedia/Private/Readers/ExrImgMediaReaderGpu.cpp
--- a/Engine/Plugins/Media/ImgMedia/Source/ImgMedia/Private/Readers/ExrImgMediaReaderGpu.cpp
+++ b/Engine/Plugins/Media/ImgMedia/Source/ImgMedia/Private/Readers/ExrImgMediaReaderGpu.cpp
@@ -510,9 +510,6 @@ void FExrImgMediaReaderGpu::CreateSampleConverterCallback
 		{
 			SCOPED_GPU_STAT(RHICmdList, ExrImgMediaReaderGpu_MipRender);
 
-			// This flag will indicate that we should wait for poll to complete.
-			BufferData->bWillBeSignaled = true;
-
 			FRHIRenderPassInfo RPInfo(RenderTargetTextureRHI, ERenderTargetActions::DontLoad_Store, nullptr, TextureMipLevel);
 			RHICmdList.BeginRenderPass(RPInfo, TEXT("ExrTextureSwizzle"));
 
@@ -571,8 +568,14 @@ void FExrImgMediaReaderGpu::CreateSampleConverterCallback
 			// Resolve render target.
 			RHICmdList.EndRenderPass();
 
-			// Mark this render command for this buffer as complete, so we can poll it and transfer later.
-			RHICmdList.WriteGPUFence(BufferData->RenderFence);
+			if (!BufferData->bWillBeSignaled)
+			{
+				// This flag will indicate that we should wait for poll to complete.
+				BufferData->bWillBeSignaled = true;
+
+				// Mark this render command for this buffer as complete, so we can poll it and transfer later.
+				RHICmdList.WriteGPUFence(BufferData->RenderFence);
+			}
 		};
 
 		int32 MipToUpscale = ConverterParams->UpscaleMip;
