ImgMedia: ExrReaderGPU: Added fence clear on buffer deletion for clarity.

diff --git a/Engine/Plugins/Media/ImgMedia/Source/ImgMedia/Private/Readers/ExrImgMediaReaderGpu.cpp b/Engine/Plugins/Media/ImgMedia/Source/ImgMedia/Private/Readers/ExrImgMediaReaderGpu.cpp
--- a/Engine/Plugins/Media/ImgMedia/Source/ImgMedia/Private/Readers/ExrImgMediaReaderGpu.cpp
+++ b/Engine/Plugins/Media/ImgMedia/Source/ImgMedia/Private/Readers/ExrImgMediaReaderGpu.cpp
@@ -111,6 +111,7 @@ FExrImgMediaReaderGpu::~FExrImgMediaReaderGpu()
 				// Check if fence has signaled.
 				check(!MemoryPoolItem->bWillBeSignaled || MemoryPoolItem->RenderFence->Poll());
 				{
+					MemoryPoolItem->RenderFence->Clear();
 					RHIUnlockBuffer(MemoryPoolItem->UploadBufferRef);
 					delete MemoryPoolItem;
 				}
@@ -707,7 +708,7 @@ void FExrImgMediaReaderGpu::ReturnGpuBufferToStagingPool(uint32 AllocSize, FStru
 		ENQUEUE_RENDER_COMMAND(DeletePooledBuffers)([Buffer](FRHICommandListImmediate& RHICmdList)
 		{
 			SCOPED_DRAW_EVENT(RHICmdList, FExrImgMediaReaderGpu_ReleaseBuffer);
-
+			Buffer->RenderFence->Clear();
 			// By this point we don't need a lock because the destructor was already called and it 
 			// is guaranteed that this buffer is no longer used anywhere else.
 			RHIUnlockBuffer(Buffer->UploadBufferRef);
