ImgMedia: Avoid upscaling sequences without mips.

#rnx

diff --git a/Engine/Plugins/Media/ImgMedia/Source/ImgMedia/Private/Assets/ImgMediaMipMapInfo.cpp b/Engine/Plugins/Media/ImgMedia/Source/ImgMedia/Private/Assets/ImgMediaMipMapInfo.cpp
--- a/Engine/Plugins/Media/ImgMedia/Source/ImgMedia/Private/Assets/ImgMediaMipMapInfo.cpp
+++ b/Engine/Plugins/Media/ImgMedia/Source/ImgMedia/Private/Assets/ImgMediaMipMapInfo.cpp
@@ -832,7 +832,7 @@ TMap<int32, FImgMediaTileSelection> FImgMediaMipMapInfo::GetVisibleTiles()
 	int32 MipToUpscale = -1;
 	static const auto CVarUpscaleMip = IConsoleManager::Get().FindTConsoleVariableDataInt(TEXT("r.ExrReaderGPU.UpscaleHigherLevelMip"));
 
-	if (CVarUpscaleMip)
+	if (SequenceInfo.NumMipLevels > 1 && CVarUpscaleMip)
 	{
 		MipToUpscale = FMath::Min(CVarUpscaleMip->GetValueOnAnyThread(), SequenceInfo.NumMipLevels - 1);
 	}
diff --git a/Engine/Plugins/Media/ImgMedia/Source/ImgMedia/Private/Readers/ExrImgMediaReaderGpu.cpp b/Engine/Plugins/Media/ImgMedia/Source/ImgMedia/Private/Readers/ExrImgMediaReaderGpu.cpp
--- a/Engine/Plugins/Media/ImgMedia/Source/ImgMedia/Private/Readers/ExrImgMediaReaderGpu.cpp
+++ b/Engine/Plugins/Media/ImgMedia/Source/ImgMedia/Private/Readers/ExrImgMediaReaderGpu.cpp
@@ -85,6 +85,20 @@ namespace {
 /* FExrImgMediaReaderGpu structors
  *****************************************************************************/
 
+FExrImgMediaReaderGpu::FExrImgMediaReaderGpu(const TSharedRef<FImgMediaLoader, ESPMode::ThreadSafe>& InLoader)
+	: FExrImgMediaReader(InLoader)
+	, LastTickedFrameCounter((uint64)-1)
+	, bIsShuttingDown(false)
+	, bFallBackToCPU(false)
+{
+	const int32 NumMipLevels = InLoader->GetNumMipLevels();
+
+	if (NumMipLevels <= 1 && CVarUpscaleLowQualityMip.GetValueOnAnyThread() >= 0)
+	{
+		UE_LOG(LogImgMedia, Display, TEXT("No upscaling for sequence without mips: %s"), *InLoader->GetImagePath(0,0));
+	}
+}
+
 FExrImgMediaReaderGpu::~FExrImgMediaReaderGpu()
 {
 	// A signal that tells all buffers that are stored in shared references not to return to the pool
@@ -307,10 +321,12 @@ bool FExrImgMediaReaderGpu::ReadFrame(int32 FrameId, const TMap<int32, FImgMedia
 		// Force mip level to be upscaled to all higher quality mips.
 		FIntPoint CurrentMipDim = ConverterParams->FullResolution;
 		TMap<int32, FImgMediaTileSelection> InMipTilesCopy = InMipTiles;
-		int32 MipToUpscale = FMath::Clamp(CVarUpscaleLowQualityMip.GetValueOnAnyThread(), -1, ConverterParams->NumMipLevels - 1);
-		ConverterParams->UpscaleMip = MipToUpscale;
-		if (MipToUpscale >= 0)
+		const int32 MipToUpscale = FMath::Clamp(CVarUpscaleLowQualityMip.GetValueOnAnyThread(), -1, ConverterParams->NumMipLevels - 1);
+
+		if (ConverterParams->NumMipLevels > 1 && MipToUpscale >= 0)
 		{
+			ConverterParams->UpscaleMip = MipToUpscale;
+
 			FImgMediaTileSelection FullSelection = FImgMediaTileSelection::CreateForTargetMipLevel(ConverterParams->FullResolution, TileDim, MipToUpscale, true);
 			if (InMipTilesCopy.Contains(MipToUpscale))
 			{
diff --git a/Engine/Plugins/Media/ImgMedia/Source/ImgMedia/Private/Readers/ExrImgMediaReaderGpu.h b/Engine/Plugins/Media/ImgMedia/Source/ImgMedia/Private/Readers/ExrImgMediaReaderGpu.h
--- a/Engine/Plugins/Media/ImgMedia/Source/ImgMedia/Private/Readers/ExrImgMediaReaderGpu.h
+++ b/Engine/Plugins/Media/ImgMedia/Source/ImgMedia/Private/Readers/ExrImgMediaReaderGpu.h
@@ -67,8 +67,7 @@ class FExrImgMediaReaderGpu
 public:
 
 	/** Default constructor. */
-	FExrImgMediaReaderGpu(const TSharedRef<FImgMediaLoader, ESPMode::ThreadSafe>& InLoader):FExrImgMediaReader(InLoader),
-		LastTickedFrameCounter((uint64)-1), bIsShuttingDown(false), bFallBackToCPU(false) {};
+	FExrImgMediaReaderGpu(const TSharedRef<FImgMediaLoader, ESPMode::ThreadSafe>& InLoader);
 	virtual ~FExrImgMediaReaderGpu();
 
 public:
