ImgMedia: Improve accuracy of mip level calculations with default anisotropic sampler.

#rnx

diff --git a/Engine/Plugins/Media/ImgMedia/Source/ImgMedia/Private/Assets/ImgMediaMipMapInfo.cpp b/Engine/Plugins/Media/ImgMedia/Source/ImgMedia/Private/Assets/ImgMediaMipMapInfo.cpp
--- a/Engine/Plugins/Media/ImgMedia/Source/ImgMedia/Private/Assets/ImgMediaMipMapInfo.cpp
+++ b/Engine/Plugins/Media/ImgMedia/Source/ImgMedia/Private/Assets/ImgMediaMipMapInfo.cpp
@@ -322,7 +322,33 @@ namespace {
 		{
 			float DistX = FVector2D::DistSquared(TexelScreenSpace[0], TexelScreenSpace[1]);
 			float DistY = FVector2D::DistSquared(TexelScreenSpace[0], TexelScreenSpace[2]);
-			OutMipLevel = 0.5f * (float)FMath::Log2(1.0f / FMath::Max(DistX, DistY)); // ~ log2(sqrt(delta))
+			OutMipLevel = 0.5f * (float)FMath::Log2(1.0f / FMath::Min(DistX, DistY)); // ~ log2(sqrt(delta))
+		}
+
+		return bValid;
+	}
+
+	// Approximates hardware mip level selection with default anisotropic filtering.
+	bool CalculateMipLevelAniso(const FImgMediaViewInfo& ViewInfo, const FVector& TexelWS, const FVector& TexelOffXWS, const FVector& TexelOffYWS, float& OutMipLevel)
+	{
+		static const float MaxAnisoLog2 = FMath::Log2(IConsoleManager::Get().FindConsoleVariable(TEXT("r.MaxAnisotropy"))->GetFloat());
+		FVector2D TexelScreenSpace[3];
+
+		bool bValid = true;
+		bValid &= ProjectWorldToScreenFast(TexelWS, ViewInfo.ViewportRect, ViewInfo.ViewProjectionMatrix, TexelScreenSpace[0]);
+		bValid &= ProjectWorldToScreenFast(TexelOffXWS, ViewInfo.ViewportRect, ViewInfo.ViewProjectionMatrix, TexelScreenSpace[1]);
+		bValid &= ProjectWorldToScreenFast(TexelOffYWS, ViewInfo.ViewportRect, ViewInfo.ViewProjectionMatrix, TexelScreenSpace[2]);
+
+		if (bValid)
+		{
+			const float Px = 1.0f / FVector2D::DistSquared(TexelScreenSpace[0], TexelScreenSpace[1]);
+			const float Py = 1.0f / FVector2D::DistSquared(TexelScreenSpace[0], TexelScreenSpace[2]);
+
+			const float MinLevel = 0.5f * FMath::Log2(FMath::Min(Px, Py));
+			const float MaxLevel = 0.5f * FMath::Log2(FMath::Max(Px, Py));
+
+			const float AnisoBias = FMath::Min(MaxLevel - MinLevel, MaxAnisoLog2);
+			OutMipLevel = MaxLevel - AnisoBias;
 		}
 
 		return bValid;
@@ -465,7 +491,7 @@ namespace {
 									float CornerStepY = TileCornerY / CurrentPartialTileNum.Y;
 									FVector CornersWS = PlaneCornerWS + (DirXWS * CornerStepX + DirYWS * CornerStepY);
 
-									if (CalculateMipLevel(ViewInfo, CornersWS, CornersWS + TexelOffsetXWS, CornersWS + TexelOffsetYWS, CalculatedLevel))
+									if (CalculateMipLevelAniso(ViewInfo, CornersWS, CornersWS + TexelOffsetXWS, CornersWS + TexelOffsetYWS, CalculatedLevel))
 									{
 										CalculatedLevel += MipMapBias + ViewInfo.MaterialTextureMipBias;
 
@@ -636,7 +662,7 @@ namespace {
 							FVector TexelOffXWS = TransformSphericalUVsToLocationWS(TileCenterUV + FVector2D(PixelDimX, 0));
 							FVector TexelOffYWS = TransformSphericalUVsToLocationWS(TileCenterUV + FVector2D(0, PixelDimY));
 
-							if (CalculateMipLevel(ViewInfo, TileCenterWS, TexelOffXWS, TexelOffYWS, CalculatedLevel))
+							if (CalculateMipLevelAniso(ViewInfo, TileCenterWS, TexelOffXWS, TexelOffYWS, CalculatedLevel))
 							{
 								CalculatedLevel += MipMapBias + ViewInfo.MaterialTextureMipBias;
 
