ImgMedia: Safety clamp over r.MaxAnisotropy missed in previous CL.

#rnx

diff --git a/Engine/Plugins/Media/ImgMedia/Source/ImgMedia/Private/Assets/ImgMediaMipMapInfo.cpp b/Engine/Plugins/Media/ImgMedia/Source/ImgMedia/Private/Assets/ImgMediaMipMapInfo.cpp
--- a/Engine/Plugins/Media/ImgMedia/Source/ImgMedia/Private/Assets/ImgMediaMipMapInfo.cpp
+++ b/Engine/Plugins/Media/ImgMedia/Source/ImgMedia/Private/Assets/ImgMediaMipMapInfo.cpp
@@ -331,7 +331,7 @@ namespace {
 	// Approximates hardware mip level selection with default anisotropic filtering.
 	bool CalculateMipLevelAniso(const FImgMediaViewInfo& ViewInfo, const FVector& TexelWS, const FVector& TexelOffXWS, const FVector& TexelOffYWS, float& OutMipLevel)
 	{
-		static const float MaxAnisoLog2 = FMath::Log2(IConsoleManager::Get().FindConsoleVariable(TEXT("r.MaxAnisotropy"))->GetFloat());
+		static const float MaxAnisoLog2 = FMath::Log2((float)FMath::Clamp(IConsoleManager::Get().FindConsoleVariable(TEXT("r.MaxAnisotropy"))->GetInt(), 1, 16));
 		FVector2D TexelScreenSpace[3];
 
 		bool bValid = true;
