ImgMedia: Prevent out-of-range frame numbers.

The tolerance offset could otherwise push the frame over the maximum index.

#rnx

diff --git a/Engine/Plugins/Media/ImgMedia/Source/ImgMedia/Private/Loader/ImgMediaLoader.cpp b/Engine/Plugins/Media/ImgMedia/Source/ImgMedia/Private/Loader/ImgMediaLoader.cpp
--- a/Engine/Plugins/Media/ImgMedia/Source/ImgMedia/Private/Loader/ImgMediaLoader.cpp
+++ b/Engine/Plugins/Media/ImgMedia/Source/ImgMedia/Private/Loader/ImgMediaLoader.cpp
@@ -1188,7 +1188,7 @@ void FImgMediaLoader::FindMips(const FString& SequencePath)
 
 uint32 FImgMediaLoader::TimeToFrameNumber(FTimespan Time) const
 {
-	if ((Time < FTimespan::Zero()) || (Time >= SequenceDuration))
+	if ((Time < FTimespan::Zero()) || (Time >= SequenceDuration) || GetNumImages() < 1)
 	{
 		return INDEX_NONE;
 	}
@@ -1201,7 +1201,7 @@ uint32 FImgMediaLoader::TimeToFrameNumber(FTimespan Time) const
 	double Frame = Time.GetTotalSeconds() * FrameDuration;
 	double Epsilon = FrameTimeErrorTollerance * FrameDuration;
 
-	return uint32(Frame + Epsilon);
+	return FMath::Min(uint32(Frame + Epsilon), uint32(GetNumImages() - 1));
 }
 
 
