MediaPlate: Allow for real-time control over mip bias (and mesh range) for visible tile calculations.

We prevent needless recreation of the MediaTextureTrackerObject on re-register calls. This allows for better use in sequencer (without having to restart the sequence for the property to apply), and consistency between editor mode and in-game mode.

#rnx

diff --git a/Engine/Plugins/Media/ImgMedia/Source/ImgMedia/Private/Assets/ImgMediaMipMapInfo.cpp b/Engine/Plugins/Media/ImgMedia/Source/ImgMedia/Private/Assets/ImgMediaMipMapInfo.cpp
--- a/Engine/Plugins/Media/ImgMedia/Source/ImgMedia/Private/Assets/ImgMediaMipMapInfo.cpp
+++ b/Engine/Plugins/Media/ImgMedia/Source/ImgMedia/Private/Assets/ImgMediaMipMapInfo.cpp
@@ -264,9 +264,9 @@ namespace {
 }
 
 
-FImgMediaMipMapObjectInfo::FImgMediaMipMapObjectInfo(UMeshComponent* InMeshComponent, float InMipMapBias)
+FImgMediaMipMapObjectInfo::FImgMediaMipMapObjectInfo(UMeshComponent* InMeshComponent, TWeakPtr<FMediaTextureTrackerObject, ESPMode::ThreadSafe> InTracker)
 	: MeshComponent(InMeshComponent)
-	, MipMapBias(InMipMapBias)
+	, Tracker(MoveTemp(InTracker))
 {
 
 }
@@ -357,8 +357,8 @@ namespace {
 	class FPlaneObjectInfo : public FImgMediaMipMapObjectInfo
 	{
 	public:
-		FPlaneObjectInfo(UMeshComponent* InMeshComponent, float InMipMapBias = 0.0f)
-			: FImgMediaMipMapObjectInfo(InMeshComponent, InMipMapBias)
+		FPlaneObjectInfo(UMeshComponent* InMeshComponent, TWeakPtr<FMediaTextureTrackerObject, ESPMode::ThreadSafe> InTracker)
+			: FImgMediaMipMapObjectInfo(InMeshComponent, MoveTemp(InTracker))
 			, PlaneSize(FVector::ZeroVector)
 		{
 			// Get size of object.
@@ -379,6 +379,14 @@ namespace {
 			{
 				return;
 			}
+			
+			TSharedPtr<FMediaTextureTrackerObject, ESPMode::ThreadSafe> ObjectInfo = Tracker.Pin();
+			if (!ObjectInfo.IsValid())
+			{
+				return;
+			}
+
+			const float MipMapBias = ObjectInfo->MipMapLODBias;
 			const float MipMapLevelPadding = FMath::Max(CVarImgMediaMipLevelPadding.GetValueOnAnyThread(), 0.0f);
 
 			const FIntPoint& SequenceTileNum = InSequenceInfo.TilingDescription.TileNum;
@@ -595,9 +603,8 @@ namespace {
 	class FSphereObjectInfo : public FImgMediaMipMapObjectInfo
 	{
 	public:
-		FSphereObjectInfo(UMeshComponent* InMeshComponent, float InMipMapBias, FVector2D InMeshRange)
-			: FImgMediaMipMapObjectInfo(InMeshComponent, InMipMapBias)
-			, MeshRange(InMeshRange)
+		FSphereObjectInfo(UMeshComponent* InMeshComponent, TWeakPtr<FMediaTextureTrackerObject, ESPMode::ThreadSafe> InTracker)
+			: FImgMediaMipMapObjectInfo(InMeshComponent, MoveTemp(InTracker))
 		{
 		}
 
@@ -609,8 +616,15 @@ namespace {
 				return;
 			}
 
-			const float MipMapLevelPadding = FMath::Max(CVarImgMediaMipLevelPadding.GetValueOnAnyThread(), 0.0f);
+			TSharedPtr<FMediaTextureTrackerObject, ESPMode::ThreadSafe> ObjectInfo = Tracker.Pin();
+			if (!ObjectInfo.IsValid())
+			{
+				return;
+			}
 
+			const FVector2D MeshRange = ObjectInfo->MeshRange;
+			const float MipMapBias = ObjectInfo->MipMapLODBias;
+			const float MipMapLevelPadding = FMath::Max(CVarImgMediaMipLevelPadding.GetValueOnAnyThread(), 0.0f);
 			const FTransform MeshTransform = Mesh->GetComponentTransform();
 			const float DefaultSphereRadius = 50.0f;
 			const int32 MaxLevel = InSequenceInfo.NumMipLevels - 1;
@@ -625,7 +639,7 @@ namespace {
 			FVector ApproxTileSizeWS = MeshTransform.GetScale3D() * (UE_TWO_PI * DefaultSphereRadius) / FMath::Max(SequencePartialTileNum.X, SequencePartialTileNum.Y);
 			float ApproxTileRadiusInWS = 0.5f * (float)FMath::Sqrt(2 * FMath::Square(ApproxTileSizeWS.GetAbsMax()));
 
-			auto TransformSphericalUVsToLocationWS = [this, &MeshTransform, DefaultSphereRadius](const FVector2D& InUV) -> FVector
+			auto TransformSphericalUVsToLocationWS = [&MeshRange, &MeshTransform, DefaultSphereRadius](const FVector2D& InUV) -> FVector
 			{
 				// Convert from latlong UV to spherical coordinates
 				FVector2D TileCornerSpherical = FVector2D(
@@ -705,9 +719,6 @@ namespace {
 				}
 			}
 		}
-
-	private:
-		FVector2D MeshRange;
 	};
 
 } //end anonymous namespace
@@ -723,26 +734,27 @@ FImgMediaMipMapInfo::~FImgMediaMipMapInfo()
 	ClearAllObjects();
 }
 
-void FImgMediaMipMapInfo::AddObject(AActor* InActor, float InMipMapBias, EMediaTextureVisibleMipsTiles InMeshMode,
-	FVector2D MeshRange)
+void FImgMediaMipMapInfo::AddObject(AActor* InActor, TWeakPtr<FMediaTextureTrackerObject, ESPMode::ThreadSafe> InTracker)
 {
-	if (InActor != nullptr)
+	TSharedPtr<FMediaTextureTrackerObject, ESPMode::ThreadSafe> Tracker = InTracker.Pin();
+
+	if (InActor != nullptr && Tracker.IsValid())
 	{
 		FScopeLock Lock(&ObjectsCriticalSection);
 
 		UMeshComponent* MeshComponent = Cast<UMeshComponent>(InActor->FindComponentByClass(UMeshComponent::StaticClass()));
 		if (MeshComponent != nullptr)
 		{
-			switch (InMeshMode)
+			switch (Tracker->VisibleMipsTilesCalculations)
 			{
 			case EMediaTextureVisibleMipsTiles::Plane:
-				Objects.Add(new FPlaneObjectInfo(MeshComponent, InMipMapBias));
+				Objects.Add(new FPlaneObjectInfo(MeshComponent, InTracker));
 				break;
 			case EMediaTextureVisibleMipsTiles::Sphere:
-				Objects.Add(new FSphereObjectInfo(MeshComponent, InMipMapBias, MeshRange));
+				Objects.Add(new FSphereObjectInfo(MeshComponent, InTracker));
 				break;
 			default:
-				Objects.Add(new FImgMediaMipMapObjectInfo(MeshComponent, InMipMapBias));
+				Objects.Add(new FImgMediaMipMapObjectInfo(MeshComponent, InTracker));
 				break;
 			}
 		}
@@ -794,7 +806,7 @@ void FImgMediaMipMapInfo::AddObjectsUsingThisMediaTexture(UMediaTexture* InMedia
 				AActor* Owner = ObjectInfo->Object.Get();
 				if (Owner != nullptr)
 				{
-					AddObject(Owner,ObjectInfo->MipMapLODBias, ObjectInfo->VisibleMipsTilesCalculations, ObjectInfo->MeshRange);
+					AddObject(Owner, ObjectInfoPtr);
 				}
 			}
 		}
diff --git a/Engine/Plugins/Media/ImgMedia/Source/ImgMedia/Private/Assets/ImgMediaMipMapInfo.h b/Engine/Plugins/Media/ImgMedia/Source/ImgMedia/Private/Assets/ImgMediaMipMapInfo.h
--- a/Engine/Plugins/Media/ImgMedia/Source/ImgMedia/Private/Assets/ImgMediaMipMapInfo.h
+++ b/Engine/Plugins/Media/ImgMedia/Source/ImgMedia/Private/Assets/ImgMediaMipMapInfo.h
@@ -185,7 +185,7 @@ private:
 class FImgMediaMipMapObjectInfo
 {
 public:
-	FImgMediaMipMapObjectInfo(UMeshComponent* InMeshComponent, float InMipMapBias = 0.0f);
+	FImgMediaMipMapObjectInfo(UMeshComponent* InMeshComponent, TWeakPtr<FMediaTextureTrackerObject, ESPMode::ThreadSafe> InTracker);
 	virtual ~FImgMediaMipMapObjectInfo() = default;
 
 	/**
@@ -206,8 +206,8 @@ public:
 protected:
 	/** Mesh component onto which the media is displayed. */
 	TWeakObjectPtr<class UMeshComponent> MeshComponent;
-	/** Mipmap level bias. */
-	float MipMapBias;
+	/** Media texture object tracker. */
+	TWeakPtr<FMediaTextureTrackerObject, ESPMode::ThreadSafe> Tracker;
 };
 
 /**
@@ -223,13 +223,9 @@ public:
 	 * This object is using our img sequence.
 	 *
 	 * @param InActor Object using our img sequence.
-	 * @param InMipMapBias Bias used for mip-level calculations.
-	 * @param InMeshMode Plane, sphere or none mesh type for mip-tile calculations.
-	 * @param MeshRange Arc size in degrees used for visible tiles calculations, specific to the sphere.
+	 * @param InObjectTracker Media texture tracker object with tile calculation information.
 	 */
-	void AddObject(AActor* InActor, float InMipMapBias = 0.0f,
-		EMediaTextureVisibleMipsTiles InMeshMode = EMediaTextureVisibleMipsTiles::None,
-		FVector2D MeshRange = FVector2D(360.0f, 180.0f));
+	void AddObject(AActor* InActor, TWeakPtr<FMediaTextureTrackerObject, ESPMode::ThreadSafe> InObjectTracker = nullptr);
 
 	/**
 	 * This object is no longer using our img sequence.
diff --git a/Engine/Plugins/Media/MediaPlate/Source/MediaPlate/Private/MediaPlate.cpp b/Engine/Plugins/Media/MediaPlate/Source/MediaPlate/Private/MediaPlate.cpp
--- a/Engine/Plugins/Media/MediaPlate/Source/MediaPlate/Private/MediaPlate.cpp
+++ b/Engine/Plugins/Media/MediaPlate/Source/MediaPlate/Private/MediaPlate.cpp
@@ -82,12 +82,6 @@ void AMediaPlate::PostRegisterAllComponents()
 {
 	Super::PostRegisterAllComponents();
 
-	if (MediaPlateComponent != nullptr)
-	{
-		// Add our media texture to the tracker.
-		MediaPlateComponent->RegisterWithMediaTextureTracker();
-	}
-
 #if WITH_EDITOR
 	AddAssetUserData();
 #endif // WITH_EDITOR
@@ -95,12 +89,6 @@ void AMediaPlate::PostRegisterAllComponents()
 
 void AMediaPlate::BeginDestroy()
 {
-	if (MediaPlateComponent != nullptr)
-	{
-		// Remove our media texture.
-		MediaPlateComponent->UnregisterWithMediaTextureTracker();
-	}
-	
 	Super::BeginDestroy();
 }
 
diff --git a/Engine/Plugins/Media/MediaPlate/Source/MediaPlate/Private/MediaPlateComponent.cpp b/Engine/Plugins/Media/MediaPlate/Source/MediaPlate/Private/MediaPlateComponent.cpp
--- a/Engine/Plugins/Media/MediaPlate/Source/MediaPlate/Private/MediaPlateComponent.cpp
+++ b/Engine/Plugins/Media/MediaPlate/Source/MediaPlate/Private/MediaPlateComponent.cpp
@@ -123,6 +123,8 @@ void UMediaPlateComponent::OnRegister()
 			SoundComponent->SetMediaPlayer(MediaPlayer);
 		}
 	}
+
+	RegisterWithMediaTextureTracker();
 }
 
 void UMediaPlateComponent::BeginPlay()
@@ -181,6 +183,13 @@ void UMediaPlateComponent::TickComponent(float DeltaTime, enum ELevelTick TickTy
 	}
 }
 
+void UMediaPlateComponent::OnUnregister()
+{
+	Super::OnUnregister();
+
+	UnregisterWithMediaTextureTracker();
+}
+
 UMediaPlayer* UMediaPlateComponent::GetMediaPlayer()
 {
 	return MediaPlayer;
@@ -358,8 +367,14 @@ void UMediaPlateComponent::PlayOnlyWhenVisibleChanged()
 
 void UMediaPlateComponent::RegisterWithMediaTextureTracker()
 {
+	UnregisterWithMediaTextureTracker();
+
 	// Set up object.
-	MediaTextureTrackerObject = MakeShared<FMediaTextureTrackerObject, ESPMode::ThreadSafe>();
+	if (MediaTextureTrackerObject == nullptr)
+	{
+		MediaTextureTrackerObject = MakeShared<FMediaTextureTrackerObject, ESPMode::ThreadSafe>();
+	}
+
 	MediaTextureTrackerObject->Object = GetOwner();
 	MediaTextureTrackerObject->MipMapLODBias = MipMapBias;
 	MediaTextureTrackerObject->VisibleMipsTilesCalculations = VisibleMipsTilesCalculations;
@@ -820,8 +835,6 @@ void UMediaPlateComponent::PostEditChangeProperty(FPropertyChangedEvent& Propert
 			MediaTextureTrackerObject->MipMapLODBias = MipMapBias;
 
 			// Note: Media texture bias and material sampler automatically updated by UMediaPlateComponent::OnRegister().
-
-			RestartPlayer();
 		}
 	}
 }
diff --git a/Engine/Plugins/Media/MediaPlate/Source/MediaPlate/Public/MediaPlateComponent.h b/Engine/Plugins/Media/MediaPlate/Source/MediaPlate/Public/MediaPlateComponent.h
--- a/Engine/Plugins/Media/MediaPlate/Source/MediaPlate/Public/MediaPlateComponent.h
+++ b/Engine/Plugins/Media/MediaPlate/Source/MediaPlate/Public/MediaPlateComponent.h
@@ -31,10 +31,11 @@ class MEDIAPLATE_API UMediaPlateComponent : public UActorComponent,
 
 public:
 	//~ UActorComponent interface.
-	virtual void OnRegister();
-	virtual void BeginPlay();
-	virtual void BeginDestroy();
+	virtual void OnRegister() override;
+	virtual void BeginPlay() override;
+	virtual void BeginDestroy() override;
 	virtual void TickComponent(float DeltaTime, enum ELevelTick TickType, FActorComponentTickFunction* ThisTickFunction) override;
+	virtual void OnUnregister() override;
 
 #if WITH_EDITOR
 	virtual void PostEditChangeProperty(FPropertyChangedEvent& PropertyChangedEvent) override;
@@ -184,15 +185,6 @@ public:
 	UFUNCTION(BlueprintCallable, Category = "Media|MediaPlateComponent")
 	void SetPlayOnlyWhenVisible(bool bInPlayOnlyWhenVisible);
 
-	/**
-	 * Adds our media texture to the media texture tracker.
-	 */
-	void RegisterWithMediaTextureTracker();
-	/**
-	 * Removes our texture from the media texture tracker.
-	 */
-	void UnregisterWithMediaTextureTracker();
-
 	/**
 	 * Call this to get the aspect ratio of the mesh.
 	 */
@@ -243,6 +235,15 @@ public:
 	virtual const FMediaSourceCacheSettings& GetCacheSettings() const override;
 
 private:
+	/**
+	 * Adds our media texture to the media texture tracker.
+	 */
+	void RegisterWithMediaTextureTracker();
+	/**
+	 * Removes our texture from the media texture tracker.
+	 */
+	void UnregisterWithMediaTextureTracker();
+
 	/**
 	 * Should be called when bPlayOnlyWhenVisible changes.
 	 */
@@ -273,7 +274,7 @@ private:
 	UPROPERTY(EditAnywhere, AdvancedDisplay, Category = "MediaPlate", meta = (AllowPrivateAccess = true))
 	EMediaTextureVisibleMipsTiles VisibleMipsTilesCalculations;
 
-	/** Media texture mip map bias shared between the (image sequence) loader and the media texture sampler. (Player restart on change.) */
+	/** Media texture mip map bias shared between the (image sequence) loader and the media texture sampler. */
 	UPROPERTY(EditAnywhere, AdvancedDisplay, Category = "MediaPlate", meta = (AllowPrivateAccess = true, UIMin = "-16.0", UIMax = "15.99"))
 	float MipMapBias = 0.0f;
 
