ImgMedia: Default worker thread count to 2.

More threads can actually hinder performance because cache updates can compete for resources with the currently needed frame. This value has been used in our test projects, so it makes sense to promote it to default.

#rnx

diff --git a/Engine/Plugins/Media/ImgMedia/Source/ImgMediaFactory/Private/ImgMediaSettings.cpp b/Engine/Plugins/Media/ImgMedia/Source/ImgMediaFactory/Private/ImgMediaSettings.cpp
--- a/Engine/Plugins/Media/ImgMedia/Source/ImgMediaFactory/Private/ImgMediaSettings.cpp
+++ b/Engine/Plugins/Media/ImgMedia/Source/ImgMediaFactory/Private/ImgMediaSettings.cpp
@@ -12,7 +12,7 @@ UImgMediaSettings::UImgMediaSettings()
 	: DefaultFrameRate(24, 1)
 	, CacheBehindPercentage(25)
 	, CacheSizeGB(1.0f)
-	, CacheThreads(8)
+	, CacheThreads(2)
 	, CacheThreadStackSizeKB(128)
 	, GlobalCacheSizeGB(1.0f)
 	, UseGlobalCache(true)
diff --git a/Engine/Plugins/Media/ImgMedia/Source/ImgMediaFactory/Public/ImgMediaSettings.h b/Engine/Plugins/Media/ImgMedia/Source/ImgMediaFactory/Public/ImgMediaSettings.h
--- a/Engine/Plugins/Media/ImgMedia/Source/ImgMediaFactory/Public/ImgMediaSettings.h
+++ b/Engine/Plugins/Media/ImgMedia/Source/ImgMediaFactory/Public/ImgMediaSettings.h
@@ -36,7 +36,7 @@ public:
 	UPROPERTY(config, EditAnywhere, Category=Caching, meta=(ClampMin=0))
 	float CacheSizeGB;
 
-	/** Maximum number of image caching threads to use (0 = number of cores, default = 8). */
+	/** Maximum number of image caching threads to use (0 = number of cores, default = 2). */
 	UPROPERTY(config, EditAnywhere, Category=Caching, meta=(ClampMin=0))
 	int32 CacheThreads;
 
