ImgMedia: Set the minimum number of frames to cache to 2. This ensures that we have one current and one cached frame ready to go.

diff --git a/Engine/Plugins/Media/ImgMedia/Source/ImgMedia/Private/Loader/ImgMediaLoader.cpp b/Engine/Plugins/Media/ImgMedia/Source/ImgMedia/Private/Loader/ImgMediaLoader.cpp
--- a/Engine/Plugins/Media/ImgMedia/Source/ImgMedia/Private/Loader/ImgMediaLoader.cpp
+++ b/Engine/Plugins/Media/ImgMedia/Source/ImgMedia/Private/Loader/ImgMediaLoader.cpp
@@ -988,6 +988,9 @@ bool FImgMediaLoader::LoadSequence(const FString& SequencePath, const FFrameRate
 	// initialize loader
 	const FPlatformMemoryStats Stats = FPlatformMemory::GetStats();
 	
+	// We always need to cache at least 2 frames - current and the future frame.
+	const int32 MinNumberOfFramesToLoad = 2;
+
 	// Are we using the smart cache?
 	if (SmartCacheSettings.bIsEnabled)
 	{
@@ -997,7 +1000,7 @@ bool FImgMediaLoader::LoadSequence(const FString& SequencePath, const FFrameRate
 		NumFramesToLoad = TimeToLookAhead / SequenceFrameRate.AsInterval();
 		// Clamp min frame to load to 1 so we always try to load something.
 		int32 NumImages = GetNumImages();
-		int32 MinFrames = FMath::Min(1, NumImages);
+		int32 MinFrames = FMath::Min(MinNumberOfFramesToLoad, NumImages);
 		NumFramesToLoad = FMath::Clamp(NumFramesToLoad, MinFrames, NumImages);
 		NumLoadBehind = 0;
 		NumLoadAhead = NumFramesToLoad;
@@ -1008,7 +1011,9 @@ bool FImgMediaLoader::LoadSequence(const FString& SequencePath, const FFrameRate
 		const SIZE_T CacheSize = FMath::Clamp(DesiredCacheSize, (SIZE_T)0, (SIZE_T)Stats.AvailablePhysical);
 
 		const int32 MaxFramesToLoad = (int32)(CacheSize / UncompressedSize);
-		NumFramesToLoad = FMath::Clamp(MaxFramesToLoad, 0, GetNumImages());
+		int32 NumImages = GetNumImages();
+		int32 MinFrames = FMath::Min(MinNumberOfFramesToLoad, NumImages);
+		NumFramesToLoad = FMath::Clamp(MaxFramesToLoad, MinFrames, NumImages);
 		const float LoadBehindScale = FMath::Clamp(Settings->CacheBehindPercentage, 0.0f, 100.0f) / 100.0f;
 
 		NumLoadBehind = (int32)(LoadBehindScale * MaxFramesToLoad);
