ImgMedia: Prevent cached frame missing tile updates from being scheduled too late.

When we detect missing tiles on previously cached frames, additional update work should not be queued for the current frame. The exr gpureader upload buffer copy (enqueued as a render command) can often arrive after the current frame was rendered into the media texture, causing visible popping tiles (even with upscaling) since its render viewports were updated but the buffer wasn't copied. If the upload buffer path is disabled, (less noticeable) tearing could occur.

#rnx

diff --git a/Engine/Plugins/Media/ImgMedia/Source/ImgMedia/Private/Loader/ImgMediaLoader.cpp b/Engine/Plugins/Media/ImgMedia/Source/ImgMedia/Private/Loader/ImgMediaLoader.cpp
--- a/Engine/Plugins/Media/ImgMedia/Source/ImgMedia/Private/Loader/ImgMediaLoader.cpp
+++ b/Engine/Plugins/Media/ImgMedia/Source/ImgMedia/Private/Loader/ImgMediaLoader.cpp
@@ -89,8 +89,8 @@ namespace ImgMediaLoader
 		return true;
 	}
 
-	// Check if [CurrentFrame] is contained in the [OffsetCount] number of frames following [OriginFrame], also taking into account looping and play direction.
-	bool IsFrameInRange(int32 CurrentFrame, int32 OriginFrame, int32 OffsetCount, int32 TotalNumFrames, float PlayRate)
+	// Check if [CurrentFrame] is contained in the [OffsetCount] number of frames after [OriginFrame], also taking into account looping and play direction.
+	bool IsCachedFrameInRange(int32 CurrentFrame, int32 OriginFrame, int32 OffsetCount, int32 TotalNumFrames, float PlayRate)
 	{
 		if (OffsetCount <= 0)
 		{
@@ -102,7 +102,8 @@ namespace ImgMediaLoader
 
 		const int32 PlayRateSign = (PlayRate >= 0.0f) ? 1 : -1;
 
-		for (int32 Offset = 0; Offset < OffsetCount; ++Offset)
+		// We start one frame away, since it's generally to late to enqueue additional update work for the current cached frame.
+		for (int32 Offset = 1; Offset <= OffsetCount; ++Offset)
 		{
 			int32 OffsetFrame = FMath::Wrap(OriginFrame + PlayRateSign * Offset, 0, TotalNumFrames - 1);
 			if (OffsetFrame == CurrentFrame)
@@ -1348,7 +1349,7 @@ void FImgMediaLoader::Update(int32 PlayHeadFrame, float PlayRate, bool Loop)
 			// No, we need one.
 			NeedFrame = true;
 		}
-		else if(ImgMediaLoader::IsFrameInRange(FrameNumber, PlayHeadFrame, FrameInvalidationMaxCount, NumImagePaths, PlayRate))
+		else if(ImgMediaLoader::IsCachedFrameInRange(FrameNumber, PlayHeadFrame, FrameInvalidationMaxCount, NumImagePaths, PlayRate))
 		{
 			// Yes. Check if we have all the desired tiles per mip level.
 			TMap<int32, FImgMediaTileSelection> DesiredMipsAndTiles;
diff --git a/Engine/Plugins/Media/ImgMedia/Source/ImgMedia/Private/Readers/ExrImgMediaReaderGpu.cpp b/Engine/Plugins/Media/ImgMedia/Source/ImgMedia/Private/Readers/ExrImgMediaReaderGpu.cpp
--- a/Engine/Plugins/Media/ImgMedia/Source/ImgMedia/Private/Readers/ExrImgMediaReaderGpu.cpp
+++ b/Engine/Plugins/Media/ImgMedia/Source/ImgMedia/Private/Readers/ExrImgMediaReaderGpu.cpp
@@ -511,6 +511,7 @@ void FExrImgMediaReaderGpu::CreateSampleConverterCallback
 {
 	auto RenderThreadSwizzler = [ConverterParams] (FRHICommandListImmediate& RHICmdList, FTexture2DRHIRef RenderTargetTextureRHI, TMap<int32, FStructuredBufferPoolItemSharedPtr>& MipBuffers)->bool
 	{
+		TRACE_CPUPROFILER_EVENT_SCOPE_TEXT(*FString::Printf(TEXT("ExrReaderGpu.Convert %d"), ConverterParams->FrameId));
 		SCOPED_DRAW_EVENT(RHICmdList, FExrImgMediaReaderGpu_Convert);
 		SCOPED_GPU_STAT(RHICmdList, ExrImgMediaReaderGpu);
 
