ImgMedia: Adding view direction as a preemptive header file change. No risk but needed for 5.1.1 hotfix.

#rnx

diff --git a/Engine/Plugins/Media/ImgMedia/Source/ImgMedia/Private/ImgMediaSceneViewExtension.cpp b/Engine/Plugins/Media/ImgMedia/Source/ImgMedia/Private/ImgMediaSceneViewExtension.cpp
--- a/Engine/Plugins/Media/ImgMedia/Source/ImgMedia/Private/ImgMediaSceneViewExtension.cpp
+++ b/Engine/Plugins/Media/ImgMedia/Source/ImgMedia/Private/ImgMediaSceneViewExtension.cpp
@@ -52,6 +52,7 @@ void FImgMediaSceneViewExtension::BeginRenderViewFamily(FSceneViewFamily& InView
 	{
 		FImgMediaViewInfo Info;
 		Info.Location = View->ViewMatrices.GetViewOrigin();
+		Info.ViewDirection = View->GetViewDirection();
 		Info.ViewProjectionMatrix = View->ViewMatrices.GetViewProjectionMatrix();
 
 		if (FMath::IsNearlyEqual(FieldOfViewMultiplier, 1.0f))
diff --git a/Engine/Plugins/Media/ImgMedia/Source/ImgMedia/Public/ImgMediaSceneViewExtension.h b/Engine/Plugins/Media/ImgMedia/Source/ImgMedia/Public/ImgMediaSceneViewExtension.h
--- a/Engine/Plugins/Media/ImgMedia/Source/ImgMedia/Public/ImgMediaSceneViewExtension.h
+++ b/Engine/Plugins/Media/ImgMedia/Source/ImgMedia/Public/ImgMediaSceneViewExtension.h
@@ -14,6 +14,8 @@ struct IMGMEDIA_API FImgMediaViewInfo
 {
 	/** Position of camera. */
 	FVector Location;
+	/** View direction of the camera. */
+	FVector ViewDirection;
 	/** View-projection matrix of the camera. */
 	FMatrix ViewProjectionMatrix;
 	/** View-projection matrix of the camera, optionally scaled for overscan frustum calculations. */
