ImgMedia: Fix spherical tile calculation for sequences with no tiles.

We calculate the spherical UV of the view direction and use it directly if it's within the current tile.

diff --git a/Engine/Plugins/Media/ImgMedia/Source/ImgMedia/Private/Assets/ImgMediaMipMapInfo.cpp b/Engine/Plugins/Media/ImgMedia/Source/ImgMedia/Private/Assets/ImgMediaMipMapInfo.cpp
--- a/Engine/Plugins/Media/ImgMedia/Source/ImgMedia/Private/Assets/ImgMediaMipMapInfo.cpp
+++ b/Engine/Plugins/Media/ImgMedia/Source/ImgMedia/Private/Assets/ImgMediaMipMapInfo.cpp
@@ -606,6 +606,7 @@ namespace {
 	public:
 		FSphereObjectInfo(UMeshComponent* InMeshComponent, TWeakPtr<FMediaTextureTrackerObject, ESPMode::ThreadSafe> InTracker)
 			: FImgMediaMipMapObjectInfo(InMeshComponent, MoveTemp(InTracker))
+			, DefaultSphereRadius(50.0f) // as defined in FMediaPlateCustomizationMesh::GenerateSphereMesh
 		{
 		}
 
@@ -623,32 +624,21 @@ namespace {
 				return;
 			}
 
-			const FVector2D MeshRange = ObjectInfo->MeshRange;
+			const FVector2f MeshRange = FVector2f(ObjectInfo->MeshRange);
 			const float MipMapBias = ObjectInfo->MipMapLODBias;
 			const float MipMapLevelPadding = FMath::Max(CVarImgMediaMipLevelPadding.GetValueOnAnyThread(), 0.0f);
 			const FTransform MeshTransform = Mesh->GetComponentTransform();
-			const float DefaultSphereRadius = 50.0f;
 			const int32 MaxLevel = InSequenceInfo.NumMipLevels - 1;
 
 			// Include all tiles containted in the visible UV region
 			const FIntPoint& SequenceTileNum = InSequenceInfo.TilingDescription.TileNum;
 			const FVector2f SequencePartialTileNum = InSequenceInfo.GetPartialTileNum();
 
-			float PixelDimX = 1.0f / InSequenceInfo.Dim.X;
-			float PixelDimY = 1.0f / InSequenceInfo.Dim.Y;
+			const float PixelDimX = 1.0f / InSequenceInfo.Dim.X;
+			const float PixelDimY = 1.0f / InSequenceInfo.Dim.Y;
 
-			FVector ApproxTileSizeWS = MeshTransform.GetScale3D() * (UE_TWO_PI * DefaultSphereRadius) / FMath::Max(SequencePartialTileNum.X, SequencePartialTileNum.Y);
-			float ApproxTileRadiusInWS = 0.5f * (float)FMath::Sqrt(2 * FMath::Square(ApproxTileSizeWS.GetAbsMax()));
-
-			auto TransformSphericalUVsToLocationWS = [&MeshRange, &MeshTransform, DefaultSphereRadius](const FVector2D& InUV) -> FVector
-			{
-				// Convert from latlong UV to spherical coordinates
-				FVector2D TileCornerSpherical = FVector2D(
-					FMath::DegreesToRadians(MeshRange.Y) * InUV.Y,
-					FMath::DegreesToRadians(MeshRange.X) * InUV.X);
-				FVector CornersWS = TileCornerSpherical.SphericalToUnitCartesian() * DefaultSphereRadius;
-				return MeshTransform.TransformPosition(CornersWS);
-			};
+			const FVector ApproxTileSizeWS = MeshTransform.GetScale3D() * (UE_TWO_PI * DefaultSphereRadius) / FMath::Max(SequencePartialTileNum.X, SequencePartialTileNum.Y);
+			const float ApproxTileRadiusInWS = 0.5f * UE_SQRT_2 * ApproxTileSizeWS.GetAbsMax();
 
 			for (const FImgMediaViewInfo& ViewInfo : InViewInfos)
 			{
@@ -661,23 +651,37 @@ namespace {
 				FConvexVolume ViewFrustum;
 				GetViewFrustumBounds(ViewFrustum, ViewInfo.OverscanViewProjectionMatrix, false, false);
 
+				// Approximated UV coordinate for a camera centered inside the sphere
+				FVector2f ViewUV = TransformDirectionWSToSphericalUVs(MeshRange, MeshTransform, ViewInfo.ViewDirection);
+
 				for (int32 TileY = 0; TileY < SequenceTileNum.Y; ++TileY)
 				{
 					for (int32 TileX = 0; TileX < SequenceTileNum.X; ++TileX)
 					{
-						FVector2D TileCenterUV = FVector2D(((float)TileX + 0.5f ) / SequencePartialTileNum.X, ((float)TileY + 0.5f) / SequencePartialTileNum.Y);
-						FVector TileCenterWS = TransformSphericalUVsToLocationWS(TileCenterUV);
+						const FVector2f TileMinCornerUV = FVector2f((float)TileX, (float)TileY) / SequencePartialTileNum;
+						const FVector2f TileMaxCornerUV = FVector2f(TileX + 1.0f, TileY + 1.0f) / SequencePartialTileNum;
+						FVector2f TileUV = 0.5f * (TileMinCornerUV + TileMaxCornerUV);
+						float CollisionSphereRadius = ApproxTileRadiusInWS;
+
+						// If the view uv is inside the tile, we use its location directly. Helpful for sequences with no tiles.
+						if (ViewUV.ComponentwiseAllGreaterOrEqual(TileMinCornerUV) && ViewUV.ComponentwiseAllLessThan(TileMaxCornerUV))
+						{
+							TileUV = ViewUV;
+							CollisionSphereRadius = 0.0f;
+						}
+
+						// The resulting location used for tile mip level estimation (either directly at the view center or its closest tile corner).
+						const FVector TileLocationWS = TransformSphericalUVsToLocationWS(MeshRange, MeshTransform, TileUV);
 
-						// For each tile corner, we include all adjacent tiles
-						if (ViewFrustum.IntersectSphere(TileCenterWS, ApproxTileRadiusInWS))
+						if (ViewFrustum.IntersectSphere(TileLocationWS, CollisionSphereRadius))
 						{
 							float CalculatedLevel;
 							FIntVector2 MipLevelRange;
 
-							FVector TexelOffXWS = TransformSphericalUVsToLocationWS(TileCenterUV + FVector2D(PixelDimX, 0));
-							FVector TexelOffYWS = TransformSphericalUVsToLocationWS(TileCenterUV + FVector2D(0, PixelDimY));
+							const FVector TexelOffXWS = TransformSphericalUVsToLocationWS(MeshRange, MeshTransform, TileUV + FVector2f(PixelDimX, 0));
+							const FVector TexelOffYWS = TransformSphericalUVsToLocationWS(MeshRange, MeshTransform, TileUV + FVector2f(0, PixelDimY));
 
-							if (CalculateMipLevelAniso(ViewInfo, TileCenterWS, TexelOffXWS, TexelOffYWS, CalculatedLevel))
+							if (CalculateMipLevelAniso(ViewInfo, TileLocationWS, TexelOffXWS, TexelOffYWS, CalculatedLevel))
 							{
 								CalculatedLevel += MipMapBias + ViewInfo.MaterialTextureMipBias;
 
@@ -707,11 +711,12 @@ namespace {
 							}
 #if false
 #if WITH_EDITOR
+							float DebugSphereRadius = FMath::Max(CollisionSphereRadius, MeshTransform.GetMaximumAxisScale());
 							// Enable this to draw a sphere where each tile is.
-							Async(EAsyncExecution::TaskGraphMainThread, [TileCenterWS, ApproxTileRadiusInWS]()
+							Async(EAsyncExecution::TaskGraphMainThread, [TileLocationWS, DebugSphereRadius]()
 								{
 									UWorld* World = GEditor->GetEditorWorldContext().World();
-									DrawDebugSphere(World, TileCenterWS, ApproxTileRadiusInWS, 8, FColor::Red, false, 0.05f);
+									DrawDebugSphere(World, TileLocationWS, DebugSphereRadius, 8, FColor::Red, false, 0.05f);
 								});
 #endif // WITH_EDITOR
 #endif // false
@@ -720,6 +725,34 @@ namespace {
 				}
 			}
 		}
+
+		FVector TransformSphericalUVsToLocationWS(const FVector2f& MeshRange, const FTransform& MeshTransform, FVector2f UV) const
+		{
+			// Scale UVs by spherical mesh range
+			UV.X *= MeshRange.X / 360.0f;
+			UV.Y = (UV.Y - 0.5f) * (MeshRange.Y / 180.0f) + 0.5f;
+			// Convert from latlong UV to spherical coordinates
+			FVector2d TileCornerSpherical = FVector2d(UE_PI * UV.Y, UE_TWO_PI * UV.X);
+			FVector CornersWS = TileCornerSpherical.SphericalToUnitCartesian() * DefaultSphereRadius;
+			return MeshTransform.TransformPosition(CornersWS);
+		}
+
+		FVector2f TransformDirectionWSToSphericalUVs(const FVector2f& MeshRange, const FTransform& MeshTransform, const FVector& InDirection) const
+		{
+			//Convert direction to spherical angular coordinates.
+			FVector SphereViewPoint = MeshTransform.InverseTransformVectorNoScale(InDirection);
+			SphereViewPoint.Normalize();
+			FVector2d Spherical = SphereViewPoint.UnitCartesianToSpherical();
+			Spherical.Y = FMath::Fmod(Spherical.Y + UE_TWO_PI, UE_TWO_PI);
+			//Convert spherical to 0-1 UV range.
+			FVector2f UV = FVector2f(Spherical.Y / UE_TWO_PI, Spherical.X / UE_PI);
+			// Scale UVs by spherical mesh range
+			UV.X /= MeshRange.X / 360.0f;
+			UV.Y = (UV.Y - 0.5f) / (MeshRange.Y / 180.0f) + 0.5f;
+			return UV;
+		}
+
+		const float DefaultSphereRadius;
 	};
 
 } //end anonymous namespace
