ImgMedia: Fixed an issue when exrs identified by relative paths are not being found and processed.

robomerge[FN_EngineMerge] -UE5-RES

diff --git a/Engine/Plugins/Media/ImgMedia/Source/ImgMediaEditor/Private/Widgets/SImgMediaProcessEXR.cpp b/Engine/Plugins/Media/ImgMedia/Source/ImgMediaEditor/Private/Widgets/SImgMediaProcessEXR.cpp
--- a/Engine/Plugins/Media/ImgMedia/Source/ImgMediaEditor/Private/Widgets/SImgMediaProcessEXR.cpp
+++ b/Engine/Plugins/Media/ImgMedia/Source/ImgMediaEditor/Private/Widgets/SImgMediaProcessEXR.cpp
@@ -224,8 +224,8 @@ void SImgMediaProcessEXR::ProcessAllImages()
 	PlatformFile.CreateDirectoryTree(*OutPath);
 
 	// Get source files.
-	FString SequencePath = FPaths::GetPath(Options->InputPath.FilePath);
-	
+	FString SequencePath = FPaths::ConvertRelativePathToFull(FPaths::GetPath(Options->InputPath.FilePath));
+
 	TArray<FString> FoundFiles;
 	IFileManager::Get().FindFiles(FoundFiles, *SequencePath, TEXT("*"));
 	FoundFiles.RemoveAll([](const FString& FileName) { return FileName.StartsWith(TEXT(".")); });
