ImgMedia: Added bandwidth throttling by skipping frames.

diff --git a/Engine/Plugins/Media/ImgMedia/Source/ImgMedia/Private/Loader/ImgMediaLoader.cpp b/Engine/Plugins/Media/ImgMedia/Source/ImgMedia/Private/Loader/ImgMediaLoader.cpp
--- a/Engine/Plugins/Media/ImgMedia/Source/ImgMedia/Private/Loader/ImgMediaLoader.cpp
+++ b/Engine/Plugins/Media/ImgMedia/Source/ImgMedia/Private/Loader/ImgMediaLoader.cpp
@@ -209,6 +209,11 @@ FImgMediaLoader::FImgMediaLoader(const TSharedRef<FImgMediaScheduler, ESPMode::T
 	, RetryCount(0)
 	, UseGlobalCache(false)
 	, SmartCacheSettings(InSmartCacheSettings)
+	, bIsPlaybackBlocking(false)
+	, bIsBandwidthThrottlingEnabled(true)
+	, bIsSkipFramesEnabled(false)
+	, SkipFramesCounter(0)
+	, SkipFramesLevel(0)
 #if WITH_EDITOR
 	, Bandwidth()
 #endif
@@ -340,6 +345,15 @@ void FImgMediaLoader::ResetFetchLogic()
 	QueuedSampleFetch.LastFrameIndex = INDEX_NONE;
 	// note: we can reset the sequence index here as this will be called when MFW does flush any queues - so we can start from scratch with no issues
 	QueuedSampleFetch.CurrentSequenceIndex = 0;
+
+	SkipFramesCounter = 0;
+	SkipFramesLevel = 0;
+}
+
+void FImgMediaLoader::SetIsPlaybackBlocking(bool bIsBlocking)
+{
+	bIsPlaybackBlocking = bIsBlocking;
+	UpdateBandwidthThrottling();
 }
 
 void FImgMediaLoader::HandlePause()
@@ -395,7 +409,7 @@ const TSharedPtr<FImgMediaFrame, ESPMode::ThreadSafe>* FImgMediaLoader::GetFrame
 }
 
 
-IMediaSamples::EFetchBestSampleResult FImgMediaLoader::FetchBestVideoSampleForTimeRange(const TRange<FMediaTimeStamp>& TimeRange, TSharedPtr<IMediaTextureSample, ESPMode::ThreadSafe>& OutSample, bool bIsLoopingEnabled, float PlayRate, bool bPlaybackIsBlocking)
+IMediaSamples::EFetchBestSampleResult FImgMediaLoader::FetchBestVideoSampleForTimeRange(const TRange<FMediaTimeStamp>& TimeRange, TSharedPtr<IMediaTextureSample, ESPMode::ThreadSafe>& OutSample, bool bIsLoopingEnabled, float PlayRate)
 {
 	if (IsInitialized() && TimeRange.HasLowerBound() && TimeRange.HasUpperBound())
 	{
@@ -493,7 +507,7 @@ IMediaSamples::EFetchBestSampleResult FImgMediaLoader::FetchBestVideoSampleForTi
 			// If playback is not blocking, we expect less expectancy of precision on the users side, but more need for speedy return of "some ok frame"
 			// So: if we detect non-blocking playback we return a "as good sample as we can", but not always the "perfect" one we calculated
 			// (still we adhere to a rough emulation of a classic output pipeline as other players have)
-			if (!bPlaybackIsBlocking)
+			if (!bIsPlaybackBlocking)
 			{
 				// Check what data we actually have in the cache already & attempt to go further backwards on the time line
 				// to get a less then optimal frame that is still on screen and available...
@@ -682,6 +696,11 @@ bool FImgMediaLoader::PeekVideoSampleTime(FMediaTimeStamp &TimeStamp, bool bIsLo
 			Idx = TimeToFrameNumber(CurrentTime);
 		}
 
+		if (SkipFramesLevel > 0)
+		{
+			Idx = GetSkipFrame(Idx);
+		}
+
 		// If possible, fetch any existing frame data...
 		// (just to see if we have any)
 		const TSharedPtr<FImgMediaFrame, ESPMode::ThreadSafe>* Frame = ((uint32)Idx != INDEX_NONE) ? (UseGlobalCache ? GlobalCache->FindAndTouch(SequenceName, Idx) : Frames.FindAndTouch(Idx)) : nullptr;
@@ -897,6 +916,8 @@ bool FImgMediaLoader::LoadSequence(const FString& SequencePath, const FFrameRate
 	}
 
 	const UImgMediaSettings* Settings = GetDefault<UImgMediaSettings>();
+	bIsBandwidthThrottlingEnabled = Settings->BandwidthThrottlingEnabled;
+	UpdateBandwidthThrottling();
 	UseGlobalCache = Settings->UseGlobalCache;
 	if (SmartCacheSettings.bIsEnabled)
 	{
@@ -1233,6 +1254,15 @@ void FImgMediaLoader::Update(int32 PlayHeadFrame, float PlayRate, bool Loop)
 		int32 LoadBehindCount = NumLoadBehind;
 		int32 LoadBehindIndex = PlayHeadFrame - FrameOffset;
 
+		// Adjust if we are skipping frames.
+		if (SkipFramesLevel > 0)
+		{
+			int32 Mult = (1 << SkipFramesLevel);
+			LoadAheadIndex = GetSkipFrame(LoadAheadIndex);
+			LoadBehindIndex = GetSkipFrame(LoadBehindIndex);
+			FrameOffset *= Mult;
+		}
+
 		// alternate between look ahead and look behind
 		
 		while ((FramesToLoad.Num() < NumFramesToLoad) &&
@@ -1245,6 +1275,7 @@ void FImgMediaLoader::Update(int32 PlayHeadFrame, float PlayRate, bool Loop)
 					if (Loop)
 					{
 						LoadAheadIndex += NumImagePaths;
+						LoadAheadIndex = GetSkipFrame(LoadAheadIndex);
 					}
 					else
 					{
@@ -1256,6 +1287,7 @@ void FImgMediaLoader::Update(int32 PlayHeadFrame, float PlayRate, bool Loop)
 					if (Loop)
 					{
 						LoadAheadIndex -= NumImagePaths;
+						LoadAheadIndex = GetSkipFrame(LoadAheadIndex);
 					}
 					else
 					{
@@ -1279,6 +1311,7 @@ void FImgMediaLoader::Update(int32 PlayHeadFrame, float PlayRate, bool Loop)
 					if (Loop)
 					{
 						LoadBehindIndex += NumImagePaths;
+						LoadBehindIndex = GetSkipFrame(LoadBehindIndex);
 					}
 					else
 					{
@@ -1290,6 +1323,7 @@ void FImgMediaLoader::Update(int32 PlayHeadFrame, float PlayRate, bool Loop)
 					if (Loop)
 					{
 						LoadBehindIndex -= NumImagePaths;
+						LoadBehindIndex = GetSkipFrame(LoadBehindIndex);
 					}
 					else
 					{
@@ -1320,6 +1354,17 @@ void FImgMediaLoader::Update(int32 PlayHeadFrame, float PlayRate, bool Loop)
 			UE_LOG(LogImgMedia, Verbose, TEXT("Loader %p: Removed Frame %i"), this, FrameNumber);
 			QueuedFrameNumbers.RemoveAtSwap(Idx);
 			Reader->CancelFrame(FrameNumber);
+			
+			// We failed to load this frame... are we throttling?
+			if (bIsSkipFramesEnabled)
+			{
+				SkipFramesCounter++;
+				if (SkipFramesCounter > SkipFramesCounterRaiseLevelThreshold)
+				{
+					SkipFramesLevel++;
+					SkipFramesCounter = SkipFramesCounterNewLevelValue;
+				}
+			}
 		}
 	}
 
@@ -1378,6 +1423,23 @@ void FImgMediaLoader::Update(int32 PlayHeadFrame, float PlayRate, bool Loop)
 	}
 	Algo::Reverse(PendingFrameNumbers);
 
+	// Are we throttling and there no frames pending?
+	if ((bIsSkipFramesEnabled) && (PendingFrameNumbers.Num() == 0))
+	{
+		// This implies that we have spare bandwidth, so lower the skip count.
+		if (SkipFramesCounter > 0)
+		{
+			SkipFramesCounter--;
+			if (SkipFramesCounter <= 0)
+			{
+				if (SkipFramesLevel > 0)
+				{
+					SkipFramesLevel--;
+					SkipFramesCounter = SkipFramesCounterNewLevelValue;
+				}
+			}
+		}
+	}
 	CSV_EVENT(ImgMedia, TEXT("LoaderUpdatePending %d %d"), (PendingFrameNumbers.Num() > 0) ? PendingFrameNumbers[0] : -1, (PendingFrameNumbers.Num() > 0) ? PendingFrameNumbers[PendingFrameNumbers.Num() - 1] : -1);
 }
 
@@ -1416,6 +1478,13 @@ void FImgMediaLoader::NotifyWorkComplete(FImgMediaLoaderWork& CompletedWork, int
 	if (QueuedFrameNumbers.Remove(FrameNumber) > 0)
 	{
 		AddFrameToCache(FrameNumber, Frame);
+
+		// Reduce SkipFramesCounter when at skip level 0, otherwise it will never go down
+		// unless we have a lot more bandwidth than needed.
+		if ((bIsSkipFramesEnabled) && (SkipFramesLevel == 0) && (SkipFramesCounter > 0))
+		{
+			SkipFramesCounter--;
+		}
 	}
 
 	WorkPool.Push(&CompletedWork);
@@ -1551,3 +1620,30 @@ bool FImgMediaLoader::GetNumberAtEndOfString(int32 &Number, const FString& Strin
 	return bFoundNumber;
 }
 
+
+int32 FImgMediaLoader::GetSkipFrame(int32 FrameIndex)
+{
+	int32 Mult = (1 << SkipFramesLevel);
+	int32 Mask = ~(Mult - 1);
+	FrameIndex = (FrameIndex & Mask) + (Mult - 1);
+	
+	const int32 NumImagePaths = GetNumImages();
+	if (FrameIndex >= NumImagePaths)
+	{
+		FrameIndex = NumImagePaths - 1;
+	}
+
+	return FrameIndex;
+}
+
+void FImgMediaLoader::UpdateBandwidthThrottling()
+{
+	bIsSkipFramesEnabled = bIsBandwidthThrottlingEnabled && (bIsPlaybackBlocking == false);
+	if (bIsSkipFramesEnabled == false)
+	{
+		SkipFramesCounter = 0;
+		SkipFramesLevel = 0;
+	}
+}
+
+
diff --git a/Engine/Plugins/Media/ImgMedia/Source/ImgMedia/Private/Loader/ImgMediaLoader.h b/Engine/Plugins/Media/ImgMedia/Source/ImgMedia/Private/Loader/ImgMediaLoader.h
--- a/Engine/Plugins/Media/ImgMedia/Source/ImgMedia/Private/Loader/ImgMediaLoader.h
+++ b/Engine/Plugins/Media/ImgMedia/Source/ImgMedia/Private/Loader/ImgMediaLoader.h
@@ -156,10 +156,9 @@ public:
 	 * @param OutSample will be filled in with the sample if found.
 	 * @param bIsLoopingEnabled True if we can loop.
 	 * @param PlayRate How fast we are playing.
-	 * @param Facade does use blocking playback to fetch samples
 	 * @return True if successful.
 	 */
-	IMediaSamples::EFetchBestSampleResult FetchBestVideoSampleForTimeRange(const TRange<FMediaTimeStamp>& TimeRange, TSharedPtr<IMediaTextureSample, ESPMode::ThreadSafe>& OutSample, bool bIsLoopingEnabled, float PlayRate, bool bPlaybackIsBlocking);
+	IMediaSamples::EFetchBestSampleResult FetchBestVideoSampleForTimeRange(const TRange<FMediaTimeStamp>& TimeRange, TSharedPtr<IMediaTextureSample, ESPMode::ThreadSafe>& OutSample, bool bIsLoopingEnabled, float PlayRate);
 
 	/**
 	 * Checks to see if a sample is available at the specificed time.
@@ -355,6 +354,11 @@ public:
 	 */
 	void ResetFetchLogic();
 
+	/**
+	 * Tell the loader if playback is blocking.
+	 */
+	void SetIsPlaybackBlocking(bool bIsBlocking);
+
 	/**
 	 * Handler for when the (owning) player's state changes to EMediaState::Paused.
 	 */
@@ -508,6 +512,24 @@ protected:
 	 */
 	bool GetNumberAtEndOfString(int32& Number, const FString& String) const;
 
+	/**
+	 * Convert FrameIndex to a frame that fits with our skipping logic.
+	 * E.g. if SkipFramesLevel is 2, then this will return frames 4 frames apart.
+	 * E.g. FrameIndex	Return value
+	 *		0			3
+	 *		1			3
+	 *		2			3
+	 *		3			3
+	 *		4			7
+	 *		5			7
+	 */
+	int32 GetSkipFrame(int32 FrameIndex);
+
+	/**
+	 * Updates how we throttle bandwidth depending on current settings.
+	 */
+	void UpdateBandwidthThrottling();
+
 private:
 
 	/** Critical section for synchronizing access to Frames. */
@@ -603,6 +625,27 @@ private:
 
 	/** Settings for the smart cache. */
 	FImgMediaLoaderSmartCacheSettings SmartCacheSettings;
+	/** True if the player is using blocked playback. */
+	bool bIsPlaybackBlocking;
+
+	/** If we want to throttle bandwidth when there is not enough bandwidth. */
+	bool bIsBandwidthThrottlingEnabled;
+	/** If we can skip frames to throttle bandwidth. */
+	bool bIsSkipFramesEnabled;
+	/** Keeps a count so we can raise/lower the skip level when this gets too high/low. */
+	int32 SkipFramesCounter;
+	/**
+	 * The current level of skipping.
+	 * The frame increment to the next frame is 1 << SkipFramesLevel.
+	 * So level 0 -> next frame is +1, i.e. the next frame, i.e. no skipping.
+	 *    level 1 -> next frame is +2, i.e. we skip 1 frame.
+	 *    level 2 -> next frame is +4, i.e. we skip 3 frames.
+	 */
+	int32 SkipFramesLevel;
+	/** When SkipFramesCounter is greater than this, then increase SkipFramesLevel. */
+	const int32 SkipFramesCounterRaiseLevelThreshold = 10;
+	/** When SkipFramesLevel changes, the value of SkipFramesCounter will reset to this. */
+	const int32 SkipFramesCounterNewLevelValue = 5;
 
 #if WITH_EDITOR
 	/** Bandwidth estimation. */
diff --git a/Engine/Plugins/Media/ImgMedia/Source/ImgMedia/Private/Player/ImgMediaPlayer.cpp b/Engine/Plugins/Media/ImgMedia/Source/ImgMedia/Private/Player/ImgMediaPlayer.cpp
--- a/Engine/Plugins/Media/ImgMedia/Source/ImgMedia/Private/Player/ImgMediaPlayer.cpp
+++ b/Engine/Plugins/Media/ImgMedia/Source/ImgMedia/Private/Player/ImgMediaPlayer.cpp
@@ -731,7 +731,14 @@ bool FImgMediaPlayer::SetRate(float Rate)
 
 void FImgMediaPlayer::SetBlockingPlaybackHint(bool FacadeWillUseBlockingPlayback)
 {
-	PlaybackIsBlocking = FacadeWillUseBlockingPlayback;
+	if (PlaybackIsBlocking != FacadeWillUseBlockingPlayback)
+	{
+		PlaybackIsBlocking = FacadeWillUseBlockingPlayback;
+		if (Loader.IsValid() && IsInitialized())
+		{
+			Loader->SetIsPlaybackBlocking(PlaybackIsBlocking);
+		}
+	}
 }
 
 /* IMediaSamples interface
@@ -792,7 +799,7 @@ IMediaSamples::EFetchBestSampleResult FImgMediaPlayer::FetchBestVideoSampleForTi
 	if (Loader.IsValid() && IsInitialized() && (CurrentState != EMediaState::Stopped))
 	{
 		// See if we have any samples in the specified time range.
-		SampleResult = Loader->FetchBestVideoSampleForTimeRange(TimeRange, OutSample, ShouldLoop, CurrentRate, PlaybackIsBlocking);
+		SampleResult = Loader->FetchBestVideoSampleForTimeRange(TimeRange, OutSample, ShouldLoop, CurrentRate);
 		Scheduler->TickInput(FTimespan::Zero(), FTimespan::MinValue());
 		if (SampleResult == IMediaSamples::EFetchBestSampleResult::Ok)
 		{
diff --git a/Engine/Plugins/Media/ImgMedia/Source/ImgMediaFactory/Private/ImgMediaSettings.cpp b/Engine/Plugins/Media/ImgMedia/Source/ImgMediaFactory/Private/ImgMediaSettings.cpp
--- a/Engine/Plugins/Media/ImgMedia/Source/ImgMediaFactory/Private/ImgMediaSettings.cpp
+++ b/Engine/Plugins/Media/ImgMedia/Source/ImgMediaFactory/Private/ImgMediaSettings.cpp
@@ -10,6 +10,7 @@
 
 UImgMediaSettings::UImgMediaSettings()
 	: DefaultFrameRate(24, 1)
+	, BandwidthThrottlingEnabled(true)
 	, CacheBehindPercentage(25)
 	, CacheSizeGB(1.0f)
 	, CacheThreads(2)
diff --git a/Engine/Plugins/Media/ImgMedia/Source/ImgMediaFactory/Public/ImgMediaSettings.h b/Engine/Plugins/Media/ImgMedia/Source/ImgMediaFactory/Public/ImgMediaSettings.h
--- a/Engine/Plugins/Media/ImgMedia/Source/ImgMediaFactory/Public/ImgMediaSettings.h
+++ b/Engine/Plugins/Media/ImgMedia/Source/ImgMediaFactory/Public/ImgMediaSettings.h
@@ -26,6 +26,10 @@ public:
 	UPROPERTY(config, EditAnywhere, Category=General)
 	FFrameRate DefaultFrameRate;
 
+	/** If there is not enough bandwidth to play the media, then lower it (for example by skipping frames) to try and get playback. */
+	UPROPERTY(config, EditAnywhere, Category = General)
+	bool BandwidthThrottlingEnabled;
+
 public:
 
 	/** Percentage of cache to use for frames behind the play head (default = 25%). */
