ImgMedia: ProcessEXR preserves data and display windows.

diff --git a/Engine/Plugins/Media/ImgMedia/Source/ImgMediaEditor/Private/Widgets/SImgMediaProcessEXR.cpp b/Engine/Plugins/Media/ImgMedia/Source/ImgMediaEditor/Private/Widgets/SImgMediaProcessEXR.cpp
--- a/Engine/Plugins/Media/ImgMedia/Source/ImgMediaEditor/Private/Widgets/SImgMediaProcessEXR.cpp
+++ b/Engine/Plugins/Media/ImgMedia/Source/ImgMediaEditor/Private/Widgets/SImgMediaProcessEXR.cpp
@@ -159,6 +159,8 @@ FReply SImgMediaProcessEXR::OnProcessImagesClicked()
 		// Are we using a player?
 		if (bUsePlayer)
 		{
+			GetGlobalImageParameters();
+
 			// Create player.
 			MediaPlayer = NewObject<UMediaPlayer>(GetTransientPackage(), "MediaPlayer", RF_Transient);
 			MediaPlayer->SetLooping(true);
@@ -248,9 +250,6 @@ void SImgMediaProcessEXR::ProcessAllImages()
 		}
 		else
 		{			
-			// ImageWrapper is always returning an alpha channel for RGB, so check if we really have one.
-			bool bHasAlphaChannel = HasAlphaChannel(Ext, FPaths::Combine(SequencePath, FoundFiles[0]));
-
 			// Get number of threads to use.
 			int32 NumThreads = Options->NumThreads;
 			if (NumThreads <= 0)
@@ -286,7 +285,7 @@ void SImgMediaProcessEXR::ProcessAllImages()
 
 				Async(EAsyncExecution::Thread, [this, SequencePath, FileName, OutPath,
 					Ext, &NumActive, InTileWidth, InTileHeight,
-					TileBorder, bEnableMips, bHasAlphaChannel]() mutable
+					TileBorder, bEnableMips]() mutable
 				{
 					FString FullFileName = FPaths::Combine(SequencePath, FileName);
 					FImage Image;
@@ -308,6 +307,9 @@ void SImgMediaProcessEXR::ProcessAllImages()
 						Image.ChangeFormat(ERawImageFormat::RGBA16F, EGammaSpace::Linear);
 					}
 
+					FImageParameters ImageParameters;
+					GetImageParameters(Ext, FullFileName, ImageParameters);
+
 					FString Name = FPaths::Combine(OutPath, FileName);
 					if (Ext != TEXT("exr"))
 					{
@@ -315,7 +317,7 @@ void SImgMediaProcessEXR::ProcessAllImages()
 					}
 
 					ProcessImageCustom(Image, InTileWidth, InTileHeight, TileBorder,
-						bEnableMips, bHasAlphaChannel, Name, false /*bUseCustomFormat*/);
+						bEnableMips, ImageParameters, Name, false /*bUseCustomFormat*/);
 
 					NumActive--;
 				});
@@ -352,22 +354,46 @@ void SImgMediaProcessEXR::ProcessAllImages()
 	});
 }
 
-bool SImgMediaProcessEXR::HasAlphaChannel(const FString& Ext, const FString& File)
+void SImgMediaProcessEXR::GetImageParameters(const FString& Ext, const FString& File,
+	FImageParameters& ImageParameters)
 {
-	bool bHasAlpha = true;
+	ImageParameters.bHasAlphaChannel = true;
+	ImageParameters.bIsDataWindowValid = false;
+	ImageParameters.bIsDisplayWindowValid = false;
+
 	// We just support EXR at the moment.
 	if (Ext == TEXT("exr"))
 	{
 		FRgbaInputFile InputFile(File);
-		bHasAlpha = InputFile.GetNumChannels() == 4;
+		ImageParameters.DataWindow = InputFile.GetDataWindowRect();
+		ImageParameters.DisplayWindow = InputFile.GetDisplayWindow();
+		ImageParameters.bHasAlphaChannel = InputFile.GetNumChannels() == 4;
+		ImageParameters.bIsDataWindowValid = true;
+		ImageParameters.bIsDisplayWindowValid = true;
 	}
+}
 
-	return bHasAlpha;
+void SImgMediaProcessEXR::GetGlobalImageParameters()
+{
+	FString SequencePath = FPaths::ConvertRelativePathToFull(FPaths::GetPath(Options->InputPath.FilePath));
+
+	TArray<FString> FoundFiles;
+	IFileManager::Get().FindFiles(FoundFiles, *SequencePath, TEXT("*"));
+	FoundFiles.RemoveAll([](const FString& FileName) { return FileName.StartsWith(TEXT(".")); });
+
+	if (FoundFiles.Num() > 0)
+	{
+		FString Ext = FPaths::GetExtension(FoundFiles[0]);
+		FString FullFileName = FPaths::Combine(SequencePath, FoundFiles[0]);
+
+		GetImageParameters(Ext, FullFileName, GlobalImageParameters);
+		GlobalImageParameters.bHasAlphaChannel = false;
+	}
 }
 
 void SImgMediaProcessEXR::ProcessImageCustom(const FImage& InImage,
 	int32 InTileWidth, int32 InTileHeight, int32 InTileBorder, bool bInEnableMips,
-	bool bHasAlphaChannel, const FString& InName,
+	const FImageParameters& ImageParameters, const FString& InName,
 	bool bIsCustomFormat)
 {
 #if IMGMEDIAEDITOR_EXR_SUPPORTED_PLATFORM
@@ -379,7 +405,7 @@ void SImgMediaProcessEXR::ProcessImageCustom(const FImage& InImage,
 
 	ProcessImageCustomRawData(RawData, Width, Height,
 		InTileWidth, InTileHeight, InTileBorder, bInEnableMips,
-		bHasAlphaChannel, InName, bIsCustomFormat);
+		ImageParameters, InName, bIsCustomFormat);
 #else // IMGMEDIAEDITOR_EXR_SUPPORTED_PLATFORM
 	UE_LOG(LogImgMediaEditor, Error, TEXT("EXR not supported on this platform."));
 #endif // IMGMEDIAEDITOR_EXR_SUPPORTED_PLATFORM
@@ -388,7 +414,7 @@ void SImgMediaProcessEXR::ProcessImageCustom(const FImage& InImage,
 void SImgMediaProcessEXR::ProcessImageCustomRawData(TArray64<uint8>& RawData,
 	int32 Width, int32 Height,
 	int32 InTileWidth, int32 InTileHeight, int32 InTileBorder, bool bInEnableMips,
-	bool bHasAlphaChannel, const FString& InName,
+	const FImageParameters& ImageParameters, const FString& InName,
 	bool bIsCustomFormat)
 {
 #if IMGMEDIAEDITOR_EXR_SUPPORTED_PLATFORM
@@ -405,7 +431,7 @@ void SImgMediaProcessEXR::ProcessImageCustomRawData(TArray64<uint8>& RawData,
 	int32 DestNumChannels = NumChannels;
 
 	// ImageWrapper always returns an alpha channel, so make sure we really have one.
-	if ((DestNumChannels == 4) && (bHasAlphaChannel == false))
+	if ((DestNumChannels == 4) && (ImageParameters.bHasAlphaChannel == false))
 	{
 		// Remove the alpha channel as its not needed.
 		RemoveAlphaChannel(RawData);
@@ -433,10 +459,14 @@ void SImgMediaProcessEXR::ProcessImageCustomRawData(TArray64<uint8>& RawData,
 	const FString AChannelName = FString(TEXT("A"));
 
 	FIntPoint Stride(2, 0);
+	FIntRect DataWindow = ImageParameters.bIsDataWindowValid ? ImageParameters.DataWindow :
+		FIntRect(0, 0, DestWidth - 1, DestHeight - 1);
+	FIntRect DisplayWindow = ImageParameters.bIsDisplayWindowValid ? ImageParameters.DisplayWindow :
+		FIntRect(0, 0, DestWidth - 1, DestHeight - 1);
 
 	// Create tiled exr file.
-	FTiledOutputFile OutFile(FIntPoint(0, 0), FIntPoint(DestWidth - 1, DestHeight - 1),
-		FIntPoint(0, 0), FIntPoint(DestWidth - 1, DestHeight - 1), bIsTiled || bIsCustomFormat);
+	FTiledOutputFile OutFile(DisplayWindow.Min, DisplayWindow.Max,
+		DataWindow.Min, DataWindow.Max, bIsTiled || bIsCustomFormat);
 
 	// Add attributes.
 	if (bIsCustomFormat)
@@ -630,7 +660,8 @@ void SImgMediaProcessEXR::ProcessImageCustomRawData(TArray64<uint8>& RawData,
 		{
 			Stride.X = BytesPerPixel;
 			Stride.Y = MipWidth * BytesPerPixel;
-			int64 BufferOffset = 0;
+			int64 BufferOffset = -ImageParameters.DataWindow.Min.X * BytesPerPixel -
+				ImageParameters.DataWindow.Min.Y * BytesPerPixel * MipWidth;
 			int64 SingleBufferOffset = BytesPerPixelPerChannel;
 			
 
@@ -655,7 +686,11 @@ void SImgMediaProcessEXR::ProcessImageCustomRawData(TArray64<uint8>& RawData,
 		}
 		else
 		{
-			OutFile.WriteTiles(0, OutFile.GetNumXTiles(MipLevel) - 1, 0, OutFile.GetNumYTiles(MipLevel) - 1, MipLevel);
+			int32 X2 = OutFile.GetNumXTiles(MipLevel) - 1;
+			X2 = FMath::Max(0, X2);
+			int32 Y2 = OutFile.GetNumYTiles(MipLevel) - 1;
+			Y2 = FMath::Max(0, Y2);
+			OutFile.WriteTiles(0, X2, 0, Y2, MipLevel);
 		}
 
 		// Switch buffers.
@@ -906,18 +941,17 @@ void SImgMediaProcessEXR::HandleProcessing()
 						int32 InTileHeight = Options->bEnableTiling ? Options->TileSizeY : 0;
 						int32 TileBorder = 0; // Note: virtual texture support is shelved for now.
 						bool bEnableMips = Options->bEnableMipMapping;
-						bool bHasAlphaChannel = false;
+						
 						FString OutPath = Options->OutputPath.Path;
 						FString FileName = FString::Printf(TEXT("image%05d.exr"), CurrentFrameIndex);
 						FString Name = FPaths::Combine(OutPath, FileName);
 
 						Async(EAsyncExecution::Thread, [this, RawData = MoveTemp(RawData), Width, Height,
-							InTileWidth, InTileHeight, TileBorder, bEnableMips,
-							bHasAlphaChannel, Name]() mutable
+							InTileWidth, InTileHeight, TileBorder, bEnableMips, Name]() mutable
 						{
 							ProcessImageCustomRawData(RawData, Width, Height,
 								InTileWidth, InTileHeight, TileBorder, bEnableMips,
-								bHasAlphaChannel, Name, false /*bUseCustomFormat*/);
+								GlobalImageParameters, Name, false /*bUseCustomFormat*/);
 						});
 					}
 					else
diff --git a/Engine/Plugins/Media/ImgMedia/Source/ImgMediaEditor/Private/Widgets/SImgMediaProcessEXR.h b/Engine/Plugins/Media/ImgMedia/Source/ImgMediaEditor/Private/Widgets/SImgMediaProcessEXR.h
--- a/Engine/Plugins/Media/ImgMedia/Source/ImgMediaEditor/Private/Widgets/SImgMediaProcessEXR.h
+++ b/Engine/Plugins/Media/ImgMedia/Source/ImgMediaEditor/Private/Widgets/SImgMediaProcessEXR.h
@@ -40,6 +40,23 @@ public:
 
 
 private:
+	/** Contains info not contained in the image data. */
+	struct FImageParameters
+	{
+		/** This window is inclusive, so (0, 0, 0, 0) is a 1 pixel rectangle at 0, 0. */
+		FIntRect DataWindow;
+		/** This window is inclusive, so(0, 0, 0, 0) is a 1 pixel rectangle at 0, 0. */
+		FIntRect DisplayWindow;
+		/** True if the image has an alpha channel. */
+		bool bHasAlphaChannel = false;
+		/** True if DataWindow can be used. */
+		bool bIsDataWindowValid = false;
+		/** True if DisplayWindow can be used. */
+		bool bIsDisplayWindowValid = false;
+	};
+	/** Contains parameters to be applied to all images. */
+	FImageParameters GlobalImageParameters;
+
 	/**
 	 * Updates our widgets based on the current state.
 	 */
@@ -56,13 +73,18 @@ private:
 	void ProcessAllImages();
 	
 	/**
-	  * Checks if this file has an alpha channel.
+	  * Gets information on the file, e.g does it have an alpha channel.
 	  *
-	  * @param Ext			File extension.
-	  * @param File			Full path to file.
-	  * @return				True if the file has an alpha channel.
+	  * @param Ext				File extension.
+	  * @param File				Full path to file.
+	  * @param ImageParameters	Will be filled in with information from the file.
 	  */
-	bool HasAlphaChannel(const FString& Ext, const FString& File);
+	void GetImageParameters(const FString& Ext, const FString& File, FImageParameters& ImageParameters);
+
+	/**
+	 * Gets parameters that can be applied to all files.
+	 */
+	void GetGlobalImageParameters();
 
 	/**
 	 * Processess a single image and writes out a file.
@@ -74,13 +96,13 @@ private:
 	 * @param InTileHeight		Desired height of tiles.
 	 * @param InTileBorder		Number of pixels to duplicate along a tile edge.
 	 * @param bInEnableMips		Turn on mip mapping.
-	 * @param bHasAlphaChannel	True if there really is an alpha channel.
+	 * @param ImageParameters	Contains information on the image.
 	 * @param InName			Full path and name of file to write.
 	 * @param bIsCustomFormat	True to output in our custom format.
 	 */
 	void ProcessImageCustom(const FImage& InImage,
 		int32 InTileWidth, int32 InTileHeight, int32 InTileBorder, 
-		bool bInEnableMips, bool bHasAlphaChannel, const FString& InName,
+		bool bInEnableMips, const FImageParameters& ImageParameters, const FString& InName,
 		bool bIsCustomFormat);
 
 	/**
@@ -95,14 +117,14 @@ private:
 	 * @param InTileHeight		Desired height of tiles.
 	 * @param InTileBorder		Number of pixels to duplicate along a tile edge.
 	 * @param bInEnableMips		Turn on mip mapping.
-	 * @param bHasAlphaChannel	True if there really is an alpha channel.
+	 * @param ImageParameters	Contains information on the image.
 	 * @param InName			Full path and name of file to write.
 	 * @param bIsCustomFormat	True to output in our custom format.
 	 */
 	void ProcessImageCustomRawData(TArray64<uint8>& RawData,
 		int32 Width, int32 Height,
 		int32 InTileWidth, int32 InTileHeight, int32 InTileBorder, bool bInEnableMips,
-		bool bHasAlphaChannel, const FString& InName, bool bIsCustomFormat);
+		const FImageParameters& ImageParameters, const FString& InName, bool bIsCustomFormat);
 
 	/**
 	 * Removes the alpha channel from a buffer.
diff --git a/Engine/Plugins/Media/ImgMedia/Source/OpenExrWrapper/Private/OpenExrWrapper.cpp b/Engine/Plugins/Media/ImgMedia/Source/OpenExrWrapper/Private/OpenExrWrapper.cpp
--- a/Engine/Plugins/Media/ImgMedia/Source/OpenExrWrapper/Private/OpenExrWrapper.cpp
+++ b/Engine/Plugins/Media/ImgMedia/Source/OpenExrWrapper/Private/OpenExrWrapper.cpp
@@ -5,6 +5,7 @@
 #include <exception>
 #include "Containers/UnrealString.h"
 #include "Logging/LogMacros.h"
+#include "Math/IntRect.h"
 #include "Modules/ModuleManager.h"
 
 PRAGMA_DEFAULT_VISIBILITY_START
@@ -127,6 +128,27 @@ FIntPoint FRgbaInputFile::GetDataWindow() const
 	);
 }
 
+FIntRect FRgbaInputFile::GetDataWindowRect() const
+{
+	Imath::Box2i Win;
+	if (InputFile != nullptr)
+	{
+		Win = ((Imf::RgbaInputFile*)InputFile)->dataWindow();
+	}
+
+	return FIntRect(Win.min.x, Win.min.y, Win.max.x, Win.max.y);
+}
+
+FIntRect FRgbaInputFile::GetDisplayWindow() const
+{
+	Imath::Box2i Win;
+	if (InputFile != nullptr)
+	{
+		Win = ((Imf::RgbaInputFile*)InputFile)->displayWindow();
+	}
+
+	return FIntRect(Win.min.x, Win.min.y, Win.max.x, Win.max.y);
+}
 
 FFrameRate FRgbaInputFile::GetFrameRate(const FFrameRate& DefaultValue) const
 {
diff --git a/Engine/Plugins/Media/ImgMedia/Source/OpenExrWrapper/Public/OpenExrWrapper.h b/Engine/Plugins/Media/ImgMedia/Source/OpenExrWrapper/Public/OpenExrWrapper.h
--- a/Engine/Plugins/Media/ImgMedia/Source/OpenExrWrapper/Public/OpenExrWrapper.h
+++ b/Engine/Plugins/Media/ImgMedia/Source/OpenExrWrapper/Public/OpenExrWrapper.h
@@ -29,6 +29,8 @@ public:
 
 	const TCHAR* GetCompressionName() const;
 	FIntPoint GetDataWindow() const;
+	FIntRect GetDataWindowRect() const;
+	FIntRect GetDisplayWindow() const;
 	FFrameRate GetFrameRate(const FFrameRate& DefaultValue) const;
 	int32 GetUncompressedSize() const;
 	int32 GetNumChannels() const;
