MediaPlate: Add support for upscaling per media plate.

#rnx

diff --git a/Engine/Plugins/Media/ImgMedia/Source/ImgMedia/Private/Assets/ImgMediaMipMapInfo.cpp b/Engine/Plugins/Media/ImgMedia/Source/ImgMedia/Private/Assets/ImgMediaMipMapInfo.cpp
--- a/Engine/Plugins/Media/ImgMedia/Source/ImgMedia/Private/Assets/ImgMediaMipMapInfo.cpp
+++ b/Engine/Plugins/Media/ImgMedia/Source/ImgMedia/Private/Assets/ImgMediaMipMapInfo.cpp
@@ -277,6 +277,16 @@ UMeshComponent* FImgMediaMipMapObjectInfo::GetMeshComponent() const
 	return MeshComponent.Get(true);
 }
 
+int32 FImgMediaMipMapObjectInfo::GetMipLevelToUpscale() const
+{
+	if (Tracker.IsValid())
+	{
+		return Tracker.Pin()->MipLevelToUpscale;
+	}
+
+	return -1;
+}
+
 void FImgMediaMipMapObjectInfo::CalculateVisibleTiles(const TArray<FImgMediaViewInfo>& InViewInfos, const FSequenceInfo& InSequenceInfo, TMap<int32, FImgMediaTileSelection>& VisibleTiles) const
 {
 	UMeshComponent* Mesh = MeshComponent.Get();
@@ -901,12 +911,11 @@ TMap<int32, FImgMediaTileSelection> FImgMediaMipMapInfo::GetVisibleTiles()
 	// accessing things that are modified by code external to this function.
 
 
-	int32 MipToUpscale = -1;
-	static const auto CVarUpscaleMip = IConsoleManager::Get().FindTConsoleVariableDataInt(TEXT("r.ExrReaderGPU.UpscaleHigherLevelMip"));
+	int32 MipToUpscale = GetMinimumMipLevelToUpscale();
 
-	if (SequenceInfo.NumMipLevels > 1 && CVarUpscaleMip)
+	if (SequenceInfo.NumMipLevels > 1)
 	{
-		MipToUpscale = FMath::Min(CVarUpscaleMip->GetValueOnAnyThread(), SequenceInfo.NumMipLevels - 1);
+		MipToUpscale = FMath::Min(MipToUpscale, SequenceInfo.NumMipLevels - 1);
 	}
 
 	FScopeLock Lock(&InfoCriticalSection);
@@ -942,6 +951,30 @@ TMap<int32, FImgMediaTileSelection> FImgMediaMipMapInfo::GetVisibleTiles()
 	return CachedVisibleTiles;
 }
 
+int32 FImgMediaMipMapInfo::GetMinimumMipLevelToUpscale() const
+{
+	int32 MinimumLevel = TNumericLimits<int32>::Max();
+
+	FScopeLock LockObjects(&ObjectsCriticalSection);
+
+	for (const FImgMediaMipMapObjectInfo* ObjectInfo : Objects)
+	{
+		const int32 ObjectMipLevelToUpscale = ObjectInfo->GetMipLevelToUpscale();
+
+		if (ObjectMipLevelToUpscale >= 0)
+		{
+			MinimumLevel = FMath::Min(MinimumLevel, ObjectMipLevelToUpscale);
+		}
+	}
+
+	if (MinimumLevel != TNumericLimits<int32>::Max())
+	{
+		return MinimumLevel;
+	}
+
+	return -1;
+}
+
 
 /**
 * Check if any scene objects are using our img sequence.
diff --git a/Engine/Plugins/Media/ImgMedia/Source/ImgMedia/Private/Assets/ImgMediaMipMapInfo.h b/Engine/Plugins/Media/ImgMedia/Source/ImgMedia/Private/Assets/ImgMediaMipMapInfo.h
--- a/Engine/Plugins/Media/ImgMedia/Source/ImgMedia/Private/Assets/ImgMediaMipMapInfo.h
+++ b/Engine/Plugins/Media/ImgMedia/Source/ImgMedia/Private/Assets/ImgMediaMipMapInfo.h
@@ -203,7 +203,12 @@ public:
 	 * @return Mesh component
 	 */
 	UMeshComponent* GetMeshComponent() const;
+
+	/** Get the mip level to upscale (-1 if disabled). */
+	int32 GetMipLevelToUpscale() const;
+
 protected:
+
 	/** Mesh component onto which the media is displayed. */
 	TWeakObjectPtr<class UMeshComponent> MeshComponent;
 	/** Media texture object tracker. */
@@ -273,6 +278,9 @@ public:
 	 * @return Mipmap level.
 	 */
 	TMap<int32, FImgMediaTileSelection> GetVisibleTiles();
+
+	/** Get the minimum mip level to upscale (-1 if disabled). */
+	int32 GetMinimumMipLevelToUpscale() const;
 	
 	/**
 	 * Check if any scene objects are using our img sequence.
diff --git a/Engine/Plugins/Media/ImgMedia/Source/ImgMedia/Private/Loader/ImgMediaLoader.cpp b/Engine/Plugins/Media/ImgMedia/Source/ImgMedia/Private/Loader/ImgMediaLoader.cpp
--- a/Engine/Plugins/Media/ImgMedia/Source/ImgMedia/Private/Loader/ImgMediaLoader.cpp
+++ b/Engine/Plugins/Media/ImgMedia/Source/ImgMedia/Private/Loader/ImgMediaLoader.cpp
@@ -761,6 +761,9 @@ IQueuedWork* FImgMediaLoader::GetWork()
 		return nullptr;
 	}
 
+	// Update minimum mip level to upscale
+	MinimumMipLevelToUpscale = GetDesiredMinimumMipLevelToUpscale();
+
 	FImgMediaLoaderWork* Work = (WorkPool.Num() > 0) ? WorkPool.Pop() : new FImgMediaLoaderWork(AsShared(), Reader.ToSharedRef());
 	
 	// Get the existing frame so we can add the mip level to it.
@@ -958,6 +961,7 @@ bool FImgMediaLoader::LoadSequence(const FString& SequencePath, const FFrameRate
 	}
 
 	NumMipLevels = FMath::Max(FirstFrameInfo.NumMipLevels, ImagePaths.Num());
+	MinimumMipLevelToUpscale = -1;
 
 	SequenceDim = FirstFrameInfo.Dim;
 
@@ -1537,6 +1541,16 @@ void FImgMediaLoader::GetDesiredMipTiles(int32 FrameIndex, TMap<int32, FImgMedia
 	}
 }
 
+int32 FImgMediaLoader::GetDesiredMinimumMipLevelToUpscale()
+{
+	if (MipMapInfo.IsValid())
+	{
+		return MipMapInfo->GetMinimumMipLevelToUpscale();
+	}
+
+	return -1;
+}
+
 
 FTimespan FImgMediaLoader::ModuloTime(FTimespan Time)
 {
diff --git a/Engine/Plugins/Media/ImgMedia/Source/ImgMedia/Private/Loader/ImgMediaLoader.h b/Engine/Plugins/Media/ImgMedia/Source/ImgMedia/Private/Loader/ImgMediaLoader.h
--- a/Engine/Plugins/Media/ImgMedia/Source/ImgMedia/Private/Loader/ImgMediaLoader.h
+++ b/Engine/Plugins/Media/ImgMedia/Source/ImgMedia/Private/Loader/ImgMediaLoader.h
@@ -259,6 +259,14 @@ public:
 		return NumMipLevels; 
 	};
 
+	/**
+	 * Get the minumum mip level to upscale to lower quality mips (-1 if disabled).
+	 */
+	int32 GetMinimumLevelToUpscale() const
+	{
+		return MinimumMipLevelToUpscale;
+	};
+
 	/**
 	 * If mips are stored in separate files we should make sure to read it.
 	 */
@@ -473,6 +481,13 @@ protected:
 	 */
 	void GetDesiredMipTiles(int32 FrameIndex, TMap<int32, FImgMediaTileSelection>& OutMipsAndTiles);
 
+	/**
+	 * Get what mip level supporting readers should upscale.
+	 *
+	 * @return minimum mip level to upscale.
+	 */
+	int32 GetDesiredMinimumMipLevelToUpscale();
+
 	/***
 	 * Modulos the time so that it is between 0 and SequenceDuration.
 	 * Handles negative numbers appropriately.
@@ -571,6 +586,9 @@ private:
 	/** Stores num of mip levels if the sequence is made out of files that contain all mips in one file. */
 	int32 NumMipLevels;
 
+	/** Minumum mip level to upscale. */
+	TAtomic<int32> MinimumMipLevelToUpscale;
+
 	/** The image sequence reader to use. */
 	TSharedPtr<IImgMediaReader, ESPMode::ThreadSafe> Reader;
 
diff --git a/Engine/Plugins/Media/ImgMedia/Source/ImgMedia/Private/Readers/ExrImgMediaReaderGpu.cpp b/Engine/Plugins/Media/ImgMedia/Source/ImgMedia/Private/Readers/ExrImgMediaReaderGpu.cpp
--- a/Engine/Plugins/Media/ImgMedia/Source/ImgMedia/Private/Readers/ExrImgMediaReaderGpu.cpp
+++ b/Engine/Plugins/Media/ImgMedia/Source/ImgMedia/Private/Readers/ExrImgMediaReaderGpu.cpp
@@ -38,13 +38,6 @@ static TAutoConsoleVariable<bool> CVarExrReaderUseUploadHeap(
 	TEXT("Utilizes upload heap and copies raw exr buffer asynchronously.\n\
 			Requires a restart of the engine."));
 
-static TAutoConsoleVariable<int32> CVarUpscaleLowQualityMip(
-	TEXT("r.ExrReaderGPU.UpscaleHigherLevelMip"),
-	-1,
-	TEXT("Upscales lower quality mips into higher quality (EX: upscaling mip 3 into mip 0, 1, 2).\n\
-			This will cover unread black regions of EXR texture with the lower quality mips.\n\
-			This will clamp to the lowest quality mip available and disabled by default (-1)"));
-
 /** This is to force user to restart after they've changed this setting. */
 static bool bUseUploadHeap = CVarExrReaderUseUploadHeap.GetValueOnAnyThread();
 
@@ -93,7 +86,7 @@ FExrImgMediaReaderGpu::FExrImgMediaReaderGpu(const TSharedRef<FImgMediaLoader, E
 {
 	const int32 NumMipLevels = InLoader->GetNumMipLevels();
 
-	if (NumMipLevels <= 1 && CVarUpscaleLowQualityMip.GetValueOnAnyThread() >= 0)
+	if (NumMipLevels <= 1 && InLoader->GetMinimumLevelToUpscale() >= 0)
 	{
 		UE_LOG(LogImgMedia, Display, TEXT("No upscaling for sequence without mips: %s"), *InLoader->GetImagePath(0,0));
 	}
@@ -321,7 +314,7 @@ bool FExrImgMediaReaderGpu::ReadFrame(int32 FrameId, const TMap<int32, FImgMedia
 		// Force mip level to be upscaled to all higher quality mips.
 		FIntPoint CurrentMipDim = ConverterParams->FullResolution;
 		TMap<int32, FImgMediaTileSelection> InMipTilesCopy = InMipTiles;
-		const int32 MipToUpscale = FMath::Clamp(CVarUpscaleLowQualityMip.GetValueOnAnyThread(), -1, ConverterParams->NumMipLevels - 1);
+		const int32 MipToUpscale = FMath::Clamp(Loader->GetMinimumLevelToUpscale(), -1, ConverterParams->NumMipLevels - 1);
 
 		if (ConverterParams->NumMipLevels > 1 && MipToUpscale >= 0)
 		{
diff --git a/Engine/Plugins/Media/MediaPlate/Source/MediaPlate/Private/MediaPlateComponent.cpp b/Engine/Plugins/Media/MediaPlate/Source/MediaPlate/Private/MediaPlateComponent.cpp
--- a/Engine/Plugins/Media/MediaPlate/Source/MediaPlate/Private/MediaPlateComponent.cpp
+++ b/Engine/Plugins/Media/MediaPlate/Source/MediaPlate/Private/MediaPlateComponent.cpp
@@ -401,6 +401,7 @@ void UMediaPlateComponent::RegisterWithMediaTextureTracker()
 	MediaTextureTrackerObject->MipMapLODBias = MipMapBias;
 	MediaTextureTrackerObject->VisibleMipsTilesCalculations = VisibleMipsTilesCalculations;
 	MediaTextureTrackerObject->MeshRange = MeshRange;
+	MediaTextureTrackerObject->MipLevelToUpscale = bEnableMipMapUpscaling ? MipLevelToUpscale : -1;
 
 	// Add our textures.
 	FMediaTextureTracker& MediaTextureTracker = FMediaTextureTracker::Get();
@@ -996,6 +997,14 @@ void UMediaPlateComponent::PostEditChangeProperty(FPropertyChangedEvent& Propert
 			// Note: Media texture bias and material sampler automatically updated by UMediaPlateComponent::OnRegister().
 		}
 	}
+	else if (PropertyChangedEvent.GetPropertyName() == GET_MEMBER_NAME_CHECKED(ThisClass, bEnableMipMapUpscaling)
+		|| PropertyChangedEvent.GetPropertyName() == GET_MEMBER_NAME_CHECKED(ThisClass, MipLevelToUpscale))
+	{
+		if (MediaTextureTrackerObject != nullptr)
+		{
+			MediaTextureTrackerObject->MipLevelToUpscale = bEnableMipMapUpscaling ? MipLevelToUpscale : -1;
+		}
+	}
 }
 
 #endif // WITH_EDITOR
diff --git a/Engine/Plugins/Media/MediaPlate/Source/MediaPlate/Public/MediaPlateComponent.h b/Engine/Plugins/Media/MediaPlate/Source/MediaPlate/Public/MediaPlateComponent.h
--- a/Engine/Plugins/Media/MediaPlate/Source/MediaPlate/Public/MediaPlateComponent.h
+++ b/Engine/Plugins/Media/MediaPlate/Source/MediaPlate/Public/MediaPlateComponent.h
@@ -307,6 +307,14 @@ private:
 	UPROPERTY(Blueprintgetter = GetIsAspectRatioAuto, BlueprintSetter = SetIsAspectRatioAuto, Category = "MediaPlate", meta = (AllowPrivateAccess = true))
 	bool bIsAspectRatioAuto = true;
 
+	/** If true then enable the use of MipLevelToUpscale as defined below. */
+	UPROPERTY(EditAnywhere, AdvancedDisplay, Category = "MediaPlate", meta = (AllowPrivateAccess = true))
+	bool bEnableMipMapUpscaling = false;
+
+	/* With exr playback, upscale into lower quality mips from this specified level. All levels including and above the specified value will be fully read. */
+	UPROPERTY(EditAnywhere, AdvancedDisplay, Category = "MediaPlate", meta = (EditCondition = "bEnableMipMapUpscaling", AllowPrivateAccess = true, UIMin = "0", UIMax = "16"))
+	int32 MipLevelToUpscale = 16;
+
 	/** If > 0, then this is the aspect ratio of our screen and 
 	 * letterboxes will be added if the media is smaller than the screen. */
 	UPROPERTY()
diff --git a/Engine/Source/Runtime/MediaAssets/Public/MediaTextureTracker.h b/Engine/Source/Runtime/MediaAssets/Public/MediaTextureTracker.h
--- a/Engine/Source/Runtime/MediaAssets/Public/MediaTextureTracker.h
+++ b/Engine/Source/Runtime/MediaAssets/Public/MediaTextureTracker.h
@@ -32,6 +32,8 @@ struct FMediaTextureTrackerObject
 	EMediaTextureVisibleMipsTiles VisibleMipsTilesCalculations;
 	/** Arc size in degrees used for visible tiles calculations, specific to the sphere.*/
 	FVector2D MeshRange;
+	/* Mip level to upscale into lower quality mips. All levels including and above the specified value will be fully read. */
+	int32 MipLevelToUpscale = -1;
 };
 
 /**
