ImgMedia: Fixed frames being cancelled when they should not be.

diff --git a/Engine/Plugins/Media/ImgMedia/Source/ImgMedia/Private/Loader/ImgMediaLoader.cpp b/Engine/Plugins/Media/ImgMedia/Source/ImgMedia/Private/Loader/ImgMediaLoader.cpp
--- a/Engine/Plugins/Media/ImgMedia/Source/ImgMedia/Private/Loader/ImgMediaLoader.cpp
+++ b/Engine/Plugins/Media/ImgMedia/Source/ImgMedia/Private/Loader/ImgMediaLoader.cpp
@@ -1490,6 +1490,10 @@ void FImgMediaLoader::NotifyWorkComplete(FImgMediaLoaderWork& CompletedWork, int
 			SkipFramesCounter--;
 		}
 	}
+	
+	// Make sure the frame is no longer cancelled, otherwise we might end up cancelling this frame
+	// later on inadvertently!
+	Reader->UncancelFrame(FrameNumber);
 
 	WorkPool.Push(&CompletedWork);
 
diff --git a/Engine/Plugins/Media/ImgMedia/Source/ImgMedia/Private/Readers/ExrImgMediaReader.cpp b/Engine/Plugins/Media/ImgMedia/Source/ImgMedia/Private/Readers/ExrImgMediaReader.cpp
--- a/Engine/Plugins/Media/ImgMedia/Source/ImgMedia/Private/Readers/ExrImgMediaReader.cpp
+++ b/Engine/Plugins/Media/ImgMedia/Source/ImgMedia/Private/Readers/ExrImgMediaReader.cpp
@@ -331,6 +331,12 @@ void FExrImgMediaReader::CancelFrame(int32 FrameNumber)
 	CanceledFrames.Add(FrameNumber);
 }
 
+void FExrImgMediaReader::UncancelFrame(int32 FrameNumber)
+{
+	FScopeLock RegionScopeLock(&CanceledFramesCriticalSection);
+	CanceledFrames.Remove(FrameNumber);
+}
+
 /** Gets reader type (GPU vs CPU) depending on size of EXR and its compression. */
 TSharedPtr<IImgMediaReader, ESPMode::ThreadSafe> FExrImgMediaReader::GetReader(const TSharedRef <FImgMediaLoader, ESPMode::ThreadSafe>& InLoader, FString FirstImageInSequencePath)
 {
diff --git a/Engine/Plugins/Media/ImgMedia/Source/ImgMedia/Private/Readers/ExrImgMediaReader.h b/Engine/Plugins/Media/ImgMedia/Source/ImgMedia/Private/Readers/ExrImgMediaReader.h
--- a/Engine/Plugins/Media/ImgMedia/Source/ImgMedia/Private/Readers/ExrImgMediaReader.h
+++ b/Engine/Plugins/Media/ImgMedia/Source/ImgMedia/Private/Readers/ExrImgMediaReader.h
@@ -31,6 +31,7 @@ public:
 	virtual bool GetFrameInfo(const FString& ImagePath, FImgMediaFrameInfo& OutInfo) override;
 	virtual bool ReadFrame(int32 FrameId, const TMap<int32, FImgMediaTileSelection>& InMipTiles, TSharedPtr<FImgMediaFrame, ESPMode::ThreadSafe> OutFrame) override;
 	virtual void CancelFrame(int32 FrameNumber) override;
+	virtual void UncancelFrame(int32 FrameNumber) override;
 
 public:
 	/** Gets reader type (GPU vs CPU) depending on size of EXR and its compression. */
diff --git a/Engine/Plugins/Media/ImgMedia/Source/ImgMedia/Private/Readers/GenericImgMediaReader.h b/Engine/Plugins/Media/ImgMedia/Source/ImgMedia/Private/Readers/GenericImgMediaReader.h
--- a/Engine/Plugins/Media/ImgMedia/Source/ImgMedia/Private/Readers/GenericImgMediaReader.h
+++ b/Engine/Plugins/Media/ImgMedia/Source/ImgMedia/Private/Readers/GenericImgMediaReader.h
@@ -36,6 +36,7 @@ public:
 	virtual bool GetFrameInfo(const FString& ImagePath, FImgMediaFrameInfo& OutInfo) override;
 	virtual bool ReadFrame(int32 FrameId, const TMap<int32, FImgMediaTileSelection>& InMipTiles, TSharedPtr<FImgMediaFrame, ESPMode::ThreadSafe> OutFrame) override;
 	virtual void CancelFrame(int32 FrameNumber) override {};
+	virtual void UncancelFrame(int32 FrameNumber) override {};
 
 private:
 
diff --git a/Engine/Plugins/Media/ImgMedia/Source/ImgMedia/Private/Readers/IImgMediaReader.h b/Engine/Plugins/Media/ImgMedia/Source/ImgMedia/Private/Readers/IImgMediaReader.h
--- a/Engine/Plugins/Media/ImgMedia/Source/ImgMedia/Private/Readers/IImgMediaReader.h
+++ b/Engine/Plugins/Media/ImgMedia/Source/ImgMedia/Private/Readers/IImgMediaReader.h
@@ -130,6 +130,12 @@ public:
 	 */
 	virtual void CancelFrame(int32 FrameNumber) = 0;
 
+	/**
+	 * Makes sure a frame that was passed to CancelFrame is no longer marked to be cancelled.
+	 * @param FrameNumber Frame that was passed to CancelFrame.
+	 */
+	virtual void UncancelFrame(int32 FrameNumber) = 0;
+
 	/**
 	 * For some readers this function allows to pre-allocate enough memory to support the
 	 * maximum number of frames with as much efficiency as possible.
