Updated a few build.cs files to remove building a few large PCH files. Saved 3.1 GBs.

diff --git a/Engine/Plugins/Media/ImgMedia/Source/ExrReaderGpu/ExrReaderGpu.Build.cs b/Engine/Plugins/Media/ImgMedia/Source/ExrReaderGpu/ExrReaderGpu.Build.cs
--- a/Engine/Plugins/Media/ImgMedia/Source/ExrReaderGpu/ExrReaderGpu.Build.cs
+++ b/Engine/Plugins/Media/ImgMedia/Source/ExrReaderGpu/ExrReaderGpu.Build.cs
@@ -15,7 +15,6 @@ namespace UnrealBuildTool.Rules
 					"CoreUObject",
 					"RenderCore",
 					"Projects",
-					"Engine",
 				});
 
 			if (Target.Platform.IsInGroup(UnrealPlatformGroup.Windows))
@@ -25,7 +24,6 @@ namespace UnrealBuildTool.Rules
 				PrivateDependencyModuleNames.Add("OpenExrWrapper");
 				PrivateIncludePaths.Add("ExrReaderGpu/Private");
 			}
-
 		}
 	}
 }
diff --git a/Engine/Source/Runtime/Experimental/Voronoi/Voronoi.Build.cs b/Engine/Source/Runtime/Experimental/Voronoi/Voronoi.Build.cs
--- a/Engine/Source/Runtime/Experimental/Voronoi/Voronoi.Build.cs
+++ b/Engine/Source/Runtime/Experimental/Voronoi/Voronoi.Build.cs
@@ -14,7 +14,7 @@ namespace UnrealBuildTool.Rules
                 }
                 );
 
-			UnsafeTypeCastWarningLevel = WarningLevel.Warning;
+			UnsafeTypeCastWarningLevel = WarningLevel.Error;
 
 			// Disable all static analysis checkers for this module ( third party code ) 
 			bDisableStaticAnalysis = true;
