Fixing ImgMedia include paths.

diff --git a/Engine/Plugins/Media/ImgMedia/Source/ImgMedia/ImgMedia.Build.cs b/Engine/Plugins/Media/ImgMedia/Source/ImgMedia/ImgMedia.Build.cs
--- a/Engine/Plugins/Media/ImgMedia/Source/ImgMedia/ImgMedia.Build.cs
+++ b/Engine/Plugins/Media/ImgMedia/Source/ImgMedia/ImgMedia.Build.cs
@@ -1,7 +1,5 @@
 // Copyright Epic Games, Inc. All Rights Reserved.
 
-using System.IO;
-
 namespace UnrealBuildTool.Rules
 {
 	public class ImgMedia : ModuleRules
@@ -50,8 +48,6 @@ namespace UnrealBuildTool.Rules
 				PrivateDependencyModuleNames.Add("ExrReaderGpu");
 			}
 
-			var EngineDir = Path.GetFullPath(Target.RelativeEnginePath);
-
 			if (Target.Platform.IsInGroup(UnrealPlatformGroup.Windows))
 			{
 				AddEngineThirdPartyPrivateStaticDependencies(Target, "DX12");
diff --git a/Engine/Plugins/Media/ImgMedia/Source/ImgMedia/Internal/Assets/ImgMediaMipMapInfo.h b/Engine/Plugins/Media/ImgMedia/Source/ImgMedia/Internal/Assets/ImgMediaMipMapInfo.h
new file mode 100644
--- /dev/null
+++ b/Engine/Plugins/Media/ImgMedia/Source/ImgMedia/Internal/Assets/ImgMediaMipMapInfo.h
@@ -0,0 +1,328 @@
+// Copyright Epic Games, Inc. All Rights Reserved.
+
+#pragma once
+
+#include "Containers/BitArray.h"
+#include "CoreMinimal.h"
+#include "IMediaOptions.h"
+#include "IMediaTextureSample.h"
+#include "ImgMediaSceneViewExtension.h"
+#include "MediaTextureTracker.h"
+#include "Tickable.h"
+#include "UObject/WeakObjectPtr.h"
+#include "UObject/WeakObjectPtrTemplates.h"
+
+class AActor;
+class UMeshComponent;
+class UMediaTexture;
+
+/**
+ * Holds size, tiling and mips sequence information.
+ */
+struct FSequenceInfo
+{
+	/** Name of this sequence. */
+	FName Name;
+	/** Pixel dimensions of this sequence. */
+	FIntPoint Dim;
+	/** Number of mip levels. */
+	int32 NumMipLevels;
+	/** Tiling description. */
+	FMediaTextureTilingDescription TilingDescription;
+	/** 
+	* Check if the sequence has tiles.
+	* 
+	* @return True if the tile counts are above one.
+	*/
+	FORCEINLINE bool IsTiled() const
+	{
+		return (TilingDescription.TileNum.X > 1) || (TilingDescription.TileNum.Y > 1);
+	}
+	/** Get the fractional number of tiles, optionally at a specified mip level.*/
+	FORCEINLINE FVector2f GetPartialTileNum(const int32 MipLevel = 0) const
+	{
+		if (IsTiled())
+		{
+			ensure(MipLevel >= 0);
+			return (FVector2f(Dim) / FVector2f(TilingDescription.TileSize)) / float(1 << MipLevel);
+		}
+
+		return FVector2f::One();
+	}
+};
+
+/**
+ * Describes which tiles are visible.
+ */
+struct FImgMediaTileSelection
+{
+	/**
+	 * Create and initialize a new tile selection.
+	 *
+	 * @param NumTilesX Horizontal tile count.
+	 * @param NumTilesY Vertical tile count.
+	 * @param bDefaultVisibility Optional visibility for the entire region.
+	 */
+	FImgMediaTileSelection(int32 NumTilesX, int32 NumTilesY, bool bDefaultVisibility = false);
+
+	/**
+	 * Create and initialize a new tile selection, adjusting tile counts for the specified higher mip level.
+	 *
+	 * @param MipZeroResolution Pixel resolution of Mip 0.
+	 * @param TileSize Dimensions of tiles in pixels.
+	 * @param TargetMipLevel Higher target mip level for the selection (usually 1 and above).
+	 * @param bDefaultVisibility Optional visibility for the entire region.
+	 */
+	static FImgMediaTileSelection CreateForTargetMipLevel(const FIntPoint& MipZeroResolution, const FIntPoint& TileSize, int32 TargetMipLevel, bool bDefaultVisibility = false);
+
+	FImgMediaTileSelection() = default;
+	~FImgMediaTileSelection() noexcept = default;
+
+	FImgMediaTileSelection(const FImgMediaTileSelection&) = default;
+	FImgMediaTileSelection& operator=(const FImgMediaTileSelection&) = default;
+
+	FImgMediaTileSelection(FImgMediaTileSelection&&) = default;
+	FImgMediaTileSelection& operator=(FImgMediaTileSelection&&) = default;
+
+	/**
+	* See if any tile is visible.
+	*
+	* @return True if so.
+	*/
+	bool IsAnyVisible() const;
+
+	/**
+	 * Check if a specific tile is visible.
+	 *
+	 * @param TileCoordX Horizontal tile coordinate.
+	 * @param TileCoordY Vertical tile coordinate.
+	 * @return True if the coordinate is visible, false otherwise.
+	 */
+	bool IsVisible(int32 TileCoordX, int32 TileCoordY) const;
+
+	/**
+	 * Check if the currently visible tiles are also visible in another selection.
+	 *
+	 * @param Other Selection to compare.
+	 * @return True if the other selection is visible, false otherwise.
+	 */
+	bool Contains(const FImgMediaTileSelection& Other) const;
+
+	/**
+	 * Mark a tile as visible.
+	 *
+	 * @param TileCoordX Horizontal tile coordinate.
+	 * @param TileCoordY Vertical tile coordinate.
+	 */
+	void SetVisible(int32 TileCoordX, int32 TileCoordY);
+
+	/**
+	 * Returns the list of visible tile coordinates, in row and column order.
+	 *
+	 * @return Visible tile coordinates array.
+	 */
+	TArray<FIntPoint> GetVisibleCoordinates() const;
+
+	/**
+	 * Returns a calculated list of contiguous visible tile regions. Only provides regions for the missing tiles if
+	 * CurrentTileSelection is specified.
+	 *
+	 * @return Visible tile regions array. 
+	 */
+	TArray<FIntRect> GetVisibleRegions(const FImgMediaTileSelection* CurrentTileSelection = nullptr) const;
+
+	/**
+	 * Return the rectangular region bounding the visible tiles.
+	 *
+	 * @return Visible tile coordinates array.
+	 */
+	FIntRect GetVisibleRegion() const;
+
+	/**
+	 * Returns the number of visible tiles.
+	 *
+	 * @return Visible tile count.
+	 */
+	int32 NumVisibleTiles() const;
+
+	/**
+	 * Returns the overall dimensions in number of tiles.
+	 *
+	 * @return FIntPoint where X is the horizontal tile count and Y the vertical tile count.
+	 */
+	FIntPoint GetDimensions() const { return Dimensions; }
+
+private:
+	/** Convert from tile coordinates to linear index. */
+	FORCEINLINE static int32 ToIndex(int32 CoordX, int32 CoordY, const FIntPoint& Dim)
+	{
+		return CoordY * Dim.X + CoordX;
+	}
+
+	/** Convert from linear index to tile coordinates. */
+	FORCEINLINE static FIntPoint FromIndex(int32 Index, const FIntPoint& Dim)
+	{
+		return FIntPoint(Index % Dim.X, Index / Dim.X);
+	}
+	
+	/** Bit array of tiles over the entire region's dimensions, true when visible. */
+	TBitArray<> Tiles;
+	
+	/** Number of tiles in each axis. */
+	FIntPoint Dimensions;
+
+	/** Cached visible tile region. */
+	mutable FIntRect CachedVisibleRegion;
+	
+	/** Flag indicating if the cached visible tile region needs to be recalculated. */
+	mutable bool bCachedVisibleRegionDirty;
+};
+
+/**
+ * Describes a single object which is using our img sequence.
+ * Base class for various objects such as planes or spheres.
+ */
+class FImgMediaMipMapObjectInfo
+{
+public:
+	FImgMediaMipMapObjectInfo(UMeshComponent* InMeshComponent, TWeakPtr<FMediaTextureTrackerObject, ESPMode::ThreadSafe> InTracker);
+	virtual ~FImgMediaMipMapObjectInfo() = default;
+
+	/**
+	 * Calculate visible tiles per mip level.
+	 *
+	 * @param InCameraInfos Active camera information
+	 * @param InSequenceInfo Active image sequence information
+	 * @param VisibleTiles Updated list of visible tiles per mip level
+	 */
+	virtual void CalculateVisibleTiles(const TArray<FImgMediaViewInfo>& InViewInfos, const FSequenceInfo& InSequenceInfo, TMap<int32, FImgMediaTileSelection>& VisibleTiles) const;
+
+	/**
+	 * Get the registered mesh component.
+	 *
+	 * @return Mesh component
+	 */
+	UMeshComponent* GetMeshComponent() const;
+
+	/** Get the mip level to upscale (-1 if disabled). */
+	int32 GetMipLevelToUpscale() const;
+
+protected:
+
+	/** Mesh component onto which the media is displayed. */
+	TWeakObjectPtr<class UMeshComponent> MeshComponent;
+	/** Media texture object tracker. */
+	TWeakPtr<FMediaTextureTrackerObject, ESPMode::ThreadSafe> Tracker;
+};
+
+/**
+ * Contains information for working with mip maps.
+ */
+class FImgMediaMipMapInfo : public IMediaOptions::FDataContainer, public FTickableGameObject
+{
+public:
+	FImgMediaMipMapInfo();
+	virtual ~FImgMediaMipMapInfo();
+
+	/**
+	 * This object is using our img sequence.
+	 *
+	 * @param InActor Object using our img sequence.
+	 * @param InObjectTracker Media texture tracker object with tile calculation information.
+	 */
+	void AddObject(AActor* InActor, TWeakPtr<FMediaTextureTrackerObject, ESPMode::ThreadSafe> InObjectTracker = nullptr);
+
+	/**
+	 * This object is no longer using our img sequence.
+	 *
+	 * @param InActor Object no longer using our img sequence.
+	 */
+	void RemoveObject(AActor* InActor);
+
+	/**
+	 * All the objects that are using this media texture will be used in our mip map calculations.
+	 *
+	 * @param InMediaTexture Media texture to get objects from.
+	 */
+	void AddObjectsUsingThisMediaTexture(UMediaTexture* InMediaTexture);
+
+	/**
+	 * Remove objects that are using this media texture from our mip map calculations.
+	 *
+	 * @param InMediaTexture Media texture to get objects from.
+	 */
+	void RemoveObjectsUsingThisMediaTexture(UMediaTexture* InMediaTexture);
+
+	/**
+	 * Remove all objects from consideration.
+	 */
+	void ClearAllObjects();
+
+	/**
+	 * Provide information on the texture needed for our image sequence.
+	 *
+	 * @param InSequenceName Name of this sequence.
+	 * @param InNumMipMaps Number of mipmaps in our image sequence
+	 * @param InSequenceDim Dimensions of the textures in our image sequence.
+	 * @param InTilingDesc Tiling description of our image sequence.
+	 */
+	void SetTextureInfo(FName InSequenceName, int32 InNumMipMaps, const FIntPoint& InSequenceDim,
+		const FMediaTextureTilingDescription& InTilingDesc);
+
+	/**
+	 * Get what mipmap level should be used.
+	 * Returns the lowest level (highest resolution) mipmap.
+	 * Assumes all higher levels will also be used.
+	 *
+	 * @param TileSelection Will be set with which tiles are visible.
+	 * @return Mipmap level.
+	 */
+	TMap<int32, FImgMediaTileSelection> GetVisibleTiles();
+
+	/** Get the minimum mip level to upscale (-1 if disabled). */
+	int32 GetMinimumMipLevelToUpscale() const;
+	
+	/**
+	 * Check if any scene objects are using our img sequence.
+	 *
+	 * @return True if any object is active.
+	 */
+	bool HasObjects() const;
+
+	//~ FTickableGameObject interface
+	virtual void Tick(float DeltaTime) override;
+	virtual TStatId GetStatId() const override { RETURN_QUICK_DECLARE_CYCLE_STAT(FImgMediaMipMapInfo, STATGROUP_Tickables); }
+	virtual bool IsTickableInEditor() const override { return true; }
+
+protected:
+
+	void SubscribeEndFrame();
+	void UnsubscribeEndFrame();
+
+	/** End frame callback registered with core delegates to receive view infos. */
+	void OnEndFrame();
+
+	/** Array of objects that are using our img sequence. */
+	TArray<FImgMediaMipMapObjectInfo*> Objects;
+
+	/** Info for each view, used in mipmap calculations. */
+	TArray<FImgMediaViewInfo> ViewInfos;
+
+	/** Size, tiling and mips sequence information. */
+	FSequenceInfo SequenceInfo;
+	
+	/** Protects info variables. */
+	FCriticalSection InfoCriticalSection;
+
+	/** Protects tracked objects. */
+	mutable FCriticalSection ObjectsCriticalSection;
+
+	/** Desired mipmap level at this current time. */
+	TMap<int32, FImgMediaTileSelection> CachedVisibleTiles;
+
+	/** True if the cached mipmap data has been calculated this frame. */
+	bool bIsCacheValid;
+
+	/** OnEndFrame core delegate handle. */
+	FDelegateHandle OnEndFrameHandle;
+};
diff --git a/Engine/Plugins/Media/ImgMedia/Source/ImgMedia/Internal/Loader/ImgMediaLoader.h b/Engine/Plugins/Media/ImgMedia/Source/ImgMedia/Internal/Loader/ImgMediaLoader.h
new file mode 100644
--- /dev/null
+++ b/Engine/Plugins/Media/ImgMedia/Source/ImgMedia/Internal/Loader/ImgMediaLoader.h
@@ -0,0 +1,672 @@
+// Copyright Epic Games, Inc. All Rights Reserved.
+
+#pragma once
+
+#include "CoreMinimal.h"
+#include "Containers/LruCache.h"
+#include "Containers/Queue.h"
+#include "IMediaSamples.h"
+#include "IMediaTextureSample.h"
+#include "IMediaView.h"
+#include "Assets/ImgMediaMipMapInfo.h"
+#include "Misc/FrameRate.h"
+#include "Templates/SharedPointer.h"
+
+class FImgMediaGlobalCache;
+class FImgMediaLoaderWork;
+class FImgMediaScheduler;
+class FImgMediaTextureSample;
+class FMediaTimeStamp;
+class IImageWrapperModule;
+class IImgMediaReader;
+class IQueuedWork;
+
+struct FImgMediaFrame;
+struct FImgMediaFrameInfo;
+
+/**
+ * Settings for the smart cache.
+ */
+struct FImgMediaLoaderSmartCacheSettings
+{
+	/** True if we are using the smart cache. */
+	bool bIsEnabled;
+	/** The cache will fill up with frames that are up to this time from the current time. */
+	float TimeToLookAhead;
+
+	/** Constructor. */
+	FImgMediaLoaderSmartCacheSettings(bool bInIsEnabled, float InTimeToLookAhead)
+		: bIsEnabled(bInIsEnabled)
+		, TimeToLookAhead(InTimeToLookAhead)
+	{
+	}
+};
+
+/**
+ * Sequence playback bandwidth requirements and estimates.
+ */
+struct FImgMediaLoaderBandwidth
+{
+	/** Number of bytes read per second from the last frame read. */
+	float Current;
+	/** Average bandwidth over time (including idle time). */
+	float Effective;
+	/** Minimum number of bytes per second needed to read all mip levels. */
+	float Required;
+
+	/** Constructor. */
+	FImgMediaLoaderBandwidth();
+
+	/** Updated calculated current and effective bandwidth numbers. */
+	void Update(const TSharedPtr<FImgMediaFrame, ESPMode::ThreadSafe>& Frame, double WorkTime);
+
+	/** Empty the read times cache and resets its index. */
+	void EmptyCache();
+
+private:
+	/** Work time cache for average bandwidth calculation. */
+	TArray<TPair<double, double>> ReadTimeCache;
+	/** Circular index for work time cache calculation. */
+	int32 ReadTimeCacheIndex;
+};
+
+/**
+ * Loads image sequence frames from disk.
+ */
+class FImgMediaLoader
+	: public TSharedFromThis<FImgMediaLoader, ESPMode::ThreadSafe>
+{
+public:
+
+	/**
+	 * Create and initialize a new instance.
+	 *
+	 * @param InScheduler The scheduler for image loading.
+	 */
+	FImgMediaLoader(const TSharedRef<FImgMediaScheduler, ESPMode::ThreadSafe>& InScheduler,
+		const TSharedRef<FImgMediaGlobalCache, ESPMode::ThreadSafe>& InGlobalCache,
+		const TSharedPtr<FImgMediaMipMapInfo, ESPMode::ThreadSafe>& InMipMapInfo,
+		bool bInFillGapsInSequence,
+		const FImgMediaLoaderSmartCacheSettings& InSmartCacheSettings);
+
+	/** Virtual destructor. */
+	virtual ~FImgMediaLoader();
+
+public:
+	/** Max number of mip map levels supported. */
+	static constexpr int32 MAX_MIPMAP_LEVELS = 32;
+
+	/**
+	 * Get the data bit rate of the video frames.
+	 *
+	 * @return Data rate (in bits per second).
+	 */
+	uint64 GetBitRate() const;
+
+#if WITH_EDITOR
+	/** Get the number of bytes read per second from the last frame read. */
+	float GetCurrentBandwidth() const { return Bandwidth.Current; }
+	/** Get the effective bandwidth. */
+	float GetEffectiveBandwidth() const { return Bandwidth.Effective; }
+	/** Get the minimum number of bytes per second needed to read all mip levels. */
+	float GetRequiredBandwidth() const { return Bandwidth.Required; }
+#endif
+
+	/**
+	 * Get the time ranges of frames that are being loaded right now.
+	 *
+	 * @param OutRangeSet Will contain the set of time ranges.
+	 * @return true on success, false otherwise.
+	 * @see GetCompletedTimeRanges, GetPendingTimeRanges
+	 */
+	void GetBusyTimeRanges(TRangeSet<FTimespan>& OutRangeSet) const;
+
+	/**
+	 * Get the time ranges of frames that are are already loaded.
+	 *
+	 * @param OutRangeSet Will contain the set of time ranges.
+	 * @return true on success, false otherwise.
+	 * @see GetBusyTimeRanges, GetPendingTimeRanges
+	 */
+	void GetCompletedTimeRanges(TRangeSet<FTimespan>& OutRangeSet) const;
+
+	/**
+	 * Get the image frame at the specified time.
+	 *
+	 * @param Time The time of the image frame to get (relative to the beginning of the sequence).
+	 * @return The frame, or nullptr if the frame wasn't available yet.
+	 */
+	TSharedPtr<FImgMediaTextureSample, ESPMode::ThreadSafe> GetFrameSample(FTimespan Time);
+
+	/**
+	 * Get the information string for the currently loaded image sequence.
+	 *
+	 * @return Info string.
+	 * @see GetSequenceDuration, GetSamples
+	 */
+	const FString& GetInfo() const
+	{
+		return Info;
+	}
+
+	/**
+	 * Tries to get the best sample for a given time range.
+	 *
+	 * @param TimeRange is the range to check for.
+	 * @param OutSample will be filled in with the sample if found.
+	 * @param bIsLoopingEnabled True if we can loop.
+	 * @param PlayRate How fast we are playing.
+	 * @return True if successful.
+	 */
+	IMediaSamples::EFetchBestSampleResult FetchBestVideoSampleForTimeRange(const TRange<FMediaTimeStamp>& TimeRange, TSharedPtr<IMediaTextureSample, ESPMode::ThreadSafe>& OutSample, bool bIsLoopingEnabled, float PlayRate);
+
+	/**
+	 * Checks to see if a sample is available at the specificed time.
+	 *
+	 * @param TimeStamp Will be filled in with the time stamp of the sample if a sample is found.
+	 * @param bIsLoopingEnabled True if we can loop.
+	 * @param PlayRate How fast we are playing.
+	 * @param CurrentTime Time to get a sample for.
+	 * @return True if a sample is available.
+	 */
+	bool PeekVideoSampleTime(FMediaTimeStamp& TimeStamp, bool bIsLoopingEnabled, float PlayRate, const FTimespan& CurrentTime);
+
+	/**
+	 * Get the time ranges of frames that are pending to be loaded.
+	 *
+	 * @param OutRangeSet Will contain the set of time ranges.
+	 * @return true on success, false otherwise.
+	 * @see GetBusyTimeRanges, GetCompletedTimeRanges
+	 */
+	void GetPendingTimeRanges(TRangeSet<FTimespan>& OutRangeSet) const;
+
+	/**
+	 * Get the image reader object used by this loader.
+	 *
+	 * @return The image reader.
+	 */
+	TSharedRef<IImgMediaReader, ESPMode::ThreadSafe> GetReader() const
+	{
+		return Reader.ToSharedRef();
+	}
+
+	/**
+	 * Get the mipmap info object used by this loader, if any.
+	 *
+	 * @return The mipmap info object.
+	 */
+	const TSharedPtr<FImgMediaMipMapInfo, ESPMode::ThreadSafe>& GetMipMapInfo() const
+	{
+		return MipMapInfo;
+	}
+
+	/**
+	 * Get the width and height of the image sequence.
+	 *
+	 * The dimensions of the image sequence are determined by reading the attributes
+	 * of the first image. The dimensions of individual image frames in the sequence
+	 * are allowed to differ. However, this usually indicates a mistake in the content
+	 * creation pipeline and will be logged out as such.
+	 *
+	 * @return Sequence dimensions.
+	 * @see GetSequenceDuration, GetSequenceFps
+	 */
+	const FIntPoint& GetSequenceDim() const
+	{
+		return SequenceDim;
+	}
+
+	/**
+	 * Get the total duration of the image sequence.
+	 *
+	 * @return Duration.
+	 * @see GetSequenceDim, GetSequenceFps
+	 */
+	FTimespan GetSequenceDuration() const
+	{
+		return SequenceDuration;
+	}
+
+	/**
+	 * Get the sequence's frame rate.
+	 *
+	 * The frame rate of the image sequence is determined by reading the attributes
+	 * of the first image. Individual image frames may specify a different frame rate,
+	 * but it will be ignored during playback.
+	 *
+	 * @return The frame rate.
+	 * @see GetSequenceDuration, GetSequenceDim
+	 */
+	FFrameRate GetSequenceFrameRate() const
+	{
+		return SequenceFrameRate;
+	}
+
+	/**
+	 * Get the path to images in the sequence.
+	 *
+	 * @param FrameNumber Frame to get.
+	 * @param MipLevel Mip level to get.
+	 * @return Returns path and filename of image.
+	 */
+	const FString& GetImagePath(int32 FrameNumber, int32 MipLevel) const;
+
+	/**
+	 * Get the number of mipmap levels we have.
+	 */
+	int32 GetNumMipLevels() const 
+	{ 
+		return NumMipLevels; 
+	};
+
+	/**
+	 * Get the minumum mip level to upscale to lower quality mips (-1 if disabled).
+	 */
+	int32 GetMinimumLevelToUpscale() const
+	{
+		return MinimumMipLevelToUpscale;
+	};
+
+	/**
+	 * If mips are stored in separate files we should make sure to read it.
+	 */
+	bool MipsInSeparateFiles() const
+	{
+		return ImagePaths.Num() > 1;
+	};
+
+	/**
+	 * Get the number of images in a single mip level.
+	 *
+	 * @return Number of images.
+	 */
+	int32 GetNumImages() const
+	{
+		return ImagePaths.Num() > 0 ? ImagePaths[0].Num() : 0;
+	}
+
+	/**
+	 * Get the number of tiles in the X direction.
+	 */
+	int32 GetNumTilesX() const
+	{
+		return TilingDescription.TileNum.X;
+	}
+
+	/**
+	 * Get the number of tiles in the Y direction.
+	 */
+	int32 GetNumTilesY() const
+	{
+		return TilingDescription.TileNum.Y;
+	}
+
+	bool IsTiled() const
+	{
+		return (TilingDescription.TileNum.X > 1) || (TilingDescription.TileNum.Y > 1);
+	}
+
+	/**
+	 * Get the next work item.
+	 *
+	 * This method is called by the scheduler.
+	 *
+	 * @return The work item, or nullptr if no work is available.
+	 */
+	IQueuedWork* GetWork();
+
+	/**
+	 * Initialize the image sequence loader.
+	 *
+	 * @param SequencePath Path to the image sequence.
+	 * @param FrameRateOverride The frame rate to use (0/0 = do not override).
+	 * @param Loop Whether the cache should loop around.
+	 * @see IsInitialized
+	 */
+	void Initialize(const FString& SequencePath, const FFrameRate& FrameRateOverride, bool Loop);
+
+	/**
+	 * Whether this loader has been initialized yet.
+	 *
+	 * @return true if initialized, false otherwise.
+	 * @see Initialize
+	 */
+	bool IsInitialized() const
+	{
+		return Initialized;
+	}
+
+	/**
+	 * Notify the loader that a work item completed.
+	 *
+	 * @param CompletedWork The work item that completed.
+	 * @param FrameNumber Number of the frame that was read.
+	 * @param Frame The frame that was read, or nullptr if reading failed.
+	 * @param WorkTime How long to read this frame (in seconds).
+	 */
+	void NotifyWorkComplete(FImgMediaLoaderWork& CompletedWork, int32 FrameNumber,
+		const TSharedPtr<FImgMediaFrame, ESPMode::ThreadSafe>& Frame, float WorkTime);
+
+	/**
+	 * Asynchronously request the image frame at the specified time.
+	 *
+	 * @param Time The time of the image frame to request (relative to the beginning of the sequence).
+	 * @param PlayRate The current play rate (used by the look-ahead logic).
+	 * @param Loop Whether the cache should loop around.
+	 * @return bool if the frame will be loaded, false otherwise.
+	 */
+	bool RequestFrame(FTimespan Time, float PlayRate, bool Loop);
+
+	/**
+	 * Reset "queued fetch" related state used to emulate player output queue behavior
+	 */
+	void ResetFetchLogic();
+
+	/**
+	 * Tell the loader if playback is blocking.
+	 */
+	void SetIsPlaybackBlocking(bool bIsBlocking);
+
+	/**
+	 * Handler for when the (owning) player's state changes to EMediaState::Paused.
+	 */
+	void HandlePause();
+
+protected:
+
+	/**
+	 * Convert a collection of frame numbers to corresponding time ranges.
+	 *
+	 * @param FrameNumbers The frame numbers to convert.
+	 * @param OutRangeSet Will contain the time ranges.
+	 */
+	void FrameNumbersToTimeRanges(const TArray<int32>& FrameNumbers, TRangeSet<FTimespan>& OutRangeSet) const;
+
+	/**
+	 * Get the play head time corresponding to the specified frame number.
+	 *
+	 * @param FrameNumber The frame number.
+	 * @return The corresponding time.
+	 * @see TimeToFrameNumber
+	 */
+	FTimespan FrameNumberToTime(uint32 FrameNumber) const;
+
+	/**
+	 * Initialize the image sequence loader.
+	 *
+	 * @param SequencePath Path to the image sequence.
+	 * @param FrameRateOverride The frame rate to use (0/0 = do not override).
+	 * @param Loop Whether the cache should loop around.
+	 * @param OutFirstFrameInfo Sequence information based on its first frame.
+	 * @return True when loading succeeds.
+	 */
+	bool LoadSequence(const FString& SequencePath, const FFrameRate& FrameRateOverride, bool Loop, FImgMediaFrameInfo& OutFirstFrameInfo);
+
+	/**
+	 * Warms up the player and start issuing requests from the start of the sequence.
+	 *
+	 * @param InFirstFrameInfo Sequence information based on its first frame.
+	 * @param Loop Whether the cache should loop around.
+	 */
+	void WarmupSequence(const FImgMediaFrameInfo& InFirstFrameInfo, bool Loop);
+
+	/**
+	 * Finds all the files in a directory and gets their path.
+	 *
+	 * @param SequencePath Directory to look in.
+	 * @param OutputPaths File paths will be added to this.
+	 */
+	void FindFiles(const FString& SequencePath, TArray<FString>& OutputPaths);
+
+	/**
+	 * Finds the mip map files for this sequence (if any).
+	 *
+	 * Typically with non mips, a single directory holds all the files of a single sequence.
+	 *
+	 * With mip maps, a directory will hold all the files of a single sequence of a specific mip level.
+	 * The naming convention is for the directory name to end in _<SIZE>.
+	 *   SIZE does not need to be a power of 2.
+	 *   Each subsequent level should have SIZE be half of the level preceeding it.
+	 *   If SIZE does not divide evenly by 2, then round down.
+	 * The part of the name preceding _<SIZE> should be the same for all mip levels.
+	 * All mip levels of the same sequence should be in the same location
+	 * E.g. /Sequence/Seq_256/, /Sequence/Seq_128/, /Sequence/Seq_64/, etc.
+
+	 * FindMips will look for mip levels that are the at the level of SequencePath and below.
+	 * E.g. if SequencePath is Seq_1024, then FindMips will look for Seq_1024, Seq_512, etc
+	 * and will NOT look for Seq_2048 even if it is present.
+	 *
+	 * @param SequencePath Directory of sequence.
+	 */
+	void FindMips(const FString& SequencePath);
+
+	/**
+	 * Get the frame number corresponding to the specified play head time.
+	 *
+	 * @param Time The play head time.
+	 * @return The corresponding frame number, or INDEX_NONE.
+	 * @see FrameNumberToTime
+	 */
+	uint32 TimeToFrameNumber(FTimespan Time) const;
+
+	/**
+	 * Update the loader based on the current play position.
+	 *
+	 * @param PlayHeadFrame Current play head frame number.
+	 * @param PlayRate Current play rate.
+	 * @param Loop Whether the cache should loop around.
+	 */
+	void Update(int32 PlayHeadFrame, float PlayRate, bool Loop);
+
+	/**
+	 * Adds an empty frame to the cache.
+	 * 
+	 * @param FrameNumber Frame number to add.
+	 */
+	void AddEmptyFrame(int32 FrameNumber);
+
+	/**
+	 * Adds a frame to the cache.
+	 *
+	 * @param FrameNumber Number of frame to add.
+	 * @param Frame Frame information to add.
+	 */
+	void AddFrameToCache(int32 FrameNumber, const TSharedPtr<FImgMediaFrame, ESPMode::ThreadSafe>& Frame);
+
+	/**
+	 * Get what mip level and tiles we should be using for a given frame.
+	 *
+	 * @param FrameIndex Frame to get mip level for.
+	 * @param OutTileSelection Will be filled in with what tiles are actually needed.
+	 */
+	void GetDesiredMipTiles(int32 FrameIndex, TMap<int32, FImgMediaTileSelection>& OutMipsAndTiles);
+
+	/**
+	 * Get what mip level supporting readers should upscale.
+	 *
+	 * @return minimum mip level to upscale.
+	 */
+	int32 GetDesiredMinimumMipLevelToUpscale();
+
+	/***
+	 * Modulos the time so that it is between 0 and SequenceDuration.
+	 * Handles negative numbers appropriately.
+	 */
+	FTimespan ModuloTime(FTimespan Time);
+
+	/**
+	 * Gets the amount of overlap (in seconds) between a frame and a time range.
+	 * A negative or zero  value indicates the frame does not overlap the range.
+	 *
+	 * @param FrameIndex Index of frame to check.
+	 * @param StartTime Start of the range.
+	 * @param EndTime End of the range.
+	 * @param Amount of overlap.
+	 */
+	float GetFrameOverlap(uint32 FrameIndex, FTimespan StartTime, FTimespan EndTime) const;
+
+	/**
+	 * Find maximum overlapping frame index for given range
+	 */
+	float FindMaxOverlapInRange(int32 StartIndex, int32 EndIndex, FTimespan StartTime, FTimespan EndTime, int32& MaxIdx) const;
+
+	/**
+	 * Get frame data for given index. If not available attempt to find earlier frame up to given index
+	 */
+	const TSharedPtr<FImgMediaFrame, ESPMode::ThreadSafe>* GetFrameForBestIndex(int32 & MaxIdx, int32 LastIndex);
+
+	/**
+	 * Helper function to get the number at the end of a string.
+	 * 
+	 * The number does not have to be at the very end,
+	 * e.g. test1.exr will get 1.
+	 * 
+	 * @param Number Will be set to the number found if one is found.
+	 * @param String String to search.
+	 * @return True if it found a number.
+	 */
+	bool GetNumberAtEndOfString(int32& Number, const FString& String) const;
+
+	/**
+	 * Convert FrameIndex to a frame that fits with our skipping logic.
+	 * E.g. if SkipFramesLevel is 2, then this will return frames 4 frames apart.
+	 * E.g. FrameIndex	Return value
+	 *		0			3
+	 *		1			3
+	 *		2			3
+	 *		3			3
+	 *		4			7
+	 *		5			7
+	 */
+	int32 GetSkipFrame(int32 FrameIndex);
+
+	/**
+	 * Updates how we throttle bandwidth depending on current settings.
+	 */
+	void UpdateBandwidthThrottling();
+
+private:
+
+	/** Critical section for synchronizing access to Frames. */
+	mutable FCriticalSection CriticalSection;
+
+	/** The currently loaded image sequence frames. */
+	TLruCache<int32, TSharedPtr<FImgMediaFrame, ESPMode::ThreadSafe>> Frames;
+
+	/** The image wrapper module to use. */
+	IImageWrapperModule& ImageWrapperModule;
+
+	/**
+	 * Paths to each image for each mip map level in the currently opened sequence.
+	 * This is an array of mip levels, and each mip level is an array of image paths.
+	 */
+	TArray<TArray<FString>> ImagePaths;
+
+	/** Media information string. */
+	FString Info;
+
+	/** Whether this loader has been initialized yet. */
+	bool Initialized;
+
+	/** If true, then any gaps in the sequence will be filled with blank frames. */
+	const bool bFillGapsInSequence;
+
+	/** Tiling description. */
+	FMediaTextureTilingDescription TilingDescription;
+
+	/** The number of frames to load ahead of the play head. */
+	int32 NumLoadAhead;
+
+	/** The number of frames to load behind the play head. */
+	int32 NumLoadBehind;
+
+	/** Total number of frames to load. */
+	int32 NumFramesToLoad;
+
+	/** Stores num of mip levels if the sequence is made out of files that contain all mips in one file. */
+	int32 NumMipLevels;
+
+	/** Minumum mip level to upscale. */
+	TAtomic<int32> MinimumMipLevelToUpscale;
+
+	/** The image sequence reader to use. */
+	TSharedPtr<IImgMediaReader, ESPMode::ThreadSafe> Reader;
+
+	/** The scheduler for image loading. */
+	TSharedPtr<FImgMediaScheduler, ESPMode::ThreadSafe> Scheduler;
+
+	/** The scheduler for image loading. */
+	TSharedPtr<FImgMediaGlobalCache, ESPMode::ThreadSafe> GlobalCache;
+
+	/** MipMapInfo object used to handle mipmaps. Could be null if we have no mipmaps. */
+	TSharedPtr<FImgMediaMipMapInfo, ESPMode::ThreadSafe> MipMapInfo;
+
+	/** Width and height of the image sequence (in pixels) .*/
+	FIntPoint SequenceDim;
+
+	/** Total length of the image sequence. */
+	FTimespan SequenceDuration;
+
+	/** Frame rate of the currently loaded sequence. */
+	FFrameRate SequenceFrameRate;
+
+	/** Identifying name of sequence files. */
+	FName SequenceName;
+
+private:
+
+	/** Index of the previously requested frame. */
+	int32 LastRequestedFrame;
+	/** Counter used to retry loading frames again. */
+	int32 RetryCount;
+
+	/** Collection of frame numbers that need to be read. */
+	TArray<int32> PendingFrameNumbers;
+
+	/** Collection of frame numbers that are being read. */
+	TArray<int32> QueuedFrameNumbers;
+
+	/** Object pool for reusable work items. */
+	TArray<FImgMediaLoaderWork*> WorkPool;
+
+	/** True if we are using the global cache, false to use the local cache. */
+	bool UseGlobalCache;
+
+	/** State related to "queue style" frame access functions */
+	struct
+	{
+		int32 LastFrameIndex;
+		uint64 CurrentSequenceIndex;
+	} QueuedSampleFetch;
+
+private:
+
+	/** Settings for the smart cache. */
+	FImgMediaLoaderSmartCacheSettings SmartCacheSettings;
+	/** True if the player is using blocked playback. */
+	bool bIsPlaybackBlocking;
+
+	/** If we want to throttle bandwidth when there is not enough bandwidth. */
+	bool bIsBandwidthThrottlingEnabled;
+	/** If we can skip frames to throttle bandwidth. */
+	bool bIsSkipFramesEnabled;
+	/** Keeps a count so we can raise/lower the skip level when this gets too high/low. */
+	int32 SkipFramesCounter;
+	/**
+	 * The current level of skipping.
+	 * The frame increment to the next frame is 1 << SkipFramesLevel.
+	 * So level 0 -> next frame is +1, i.e. the next frame, i.e. no skipping.
+	 *    level 1 -> next frame is +2, i.e. we skip 1 frame.
+	 *    level 2 -> next frame is +4, i.e. we skip 3 frames.
+	 */
+	int32 SkipFramesLevel;
+	/** When SkipFramesCounter is greater than this, then increase SkipFramesLevel. */
+	const int32 SkipFramesCounterRaiseLevelThreshold = 10;
+	/** When SkipFramesLevel changes, the value of SkipFramesCounter will reset to this. */
+	const int32 SkipFramesCounterNewLevelValue = 5;
+
+#if WITH_EDITOR
+	/** Bandwidth estimation. */
+	FImgMediaLoaderBandwidth Bandwidth;
+#endif
+};
diff --git a/Engine/Plugins/Media/ImgMedia/Source/ImgMedia/Internal/Player/ImgMediaPlayer.h b/Engine/Plugins/Media/ImgMedia/Source/ImgMedia/Internal/Player/ImgMediaPlayer.h
new file mode 100644
--- /dev/null
+++ b/Engine/Plugins/Media/ImgMedia/Source/ImgMedia/Internal/Player/ImgMediaPlayer.h
@@ -0,0 +1,179 @@
+// Copyright Epic Games, Inc. All Rights Reserved.
+
+#pragma once
+
+#include "CoreMinimal.h"
+#include "IMediaCache.h"
+#include "IMediaControls.h"
+#include "IMediaPlayer.h"
+#include "IMediaSamples.h"
+#include "IMediaTracks.h"
+#include "IMediaView.h"
+
+class FImgMediaLoader;
+class FImgMediaScheduler;
+class IImgMediaReader;
+class IMediaEventSink;
+class IMediaTextureSample;
+class FImgMediaGlobalCache;
+
+
+/**
+ * Implements a media player for image sequences.
+ */
+class FImgMediaPlayer
+	: public IMediaPlayer
+	, protected IMediaCache
+	, protected IMediaControls
+	, protected IMediaSamples
+	, protected IMediaTracks
+	, public IMediaView
+{
+public:
+
+	/**
+	 * Create and initialize a new instance.
+	 *
+	 * @param InEventSink The object that receives media events from this player.
+	 * @param InScheduler The image loading scheduler to use.
+	 */
+	FImgMediaPlayer(IMediaEventSink& InEventSink, const TSharedRef<FImgMediaScheduler, ESPMode::ThreadSafe>& InScheduler,
+		const TSharedRef<FImgMediaGlobalCache, ESPMode::ThreadSafe>& InGlobalCache);
+
+	/** Virtual destructor. */
+	virtual ~FImgMediaPlayer();
+
+	/** Get the loader. */
+	TSharedPtr<FImgMediaLoader, ESPMode::ThreadSafe> GetLoader() { return Loader; }
+
+public:
+
+	//~ IMediaPlayer interface
+
+	virtual void Close() override;
+	virtual IMediaCache& GetCache() override;
+	virtual IMediaControls& GetControls() override;
+	virtual FString GetInfo() const override;
+	virtual FVariant GetMediaInfo(FName InfoName) const override;
+	virtual FGuid GetPlayerPluginGUID() const override;
+	virtual IMediaSamples& GetSamples() override;
+	virtual FString GetStats() const override;
+	virtual IMediaTracks& GetTracks() override;
+	virtual FString GetUrl() const override;
+	virtual IMediaView& GetView() override;
+	virtual bool Open(const FString& Url, const IMediaOptions* Options) override;
+	virtual bool Open(const TSharedRef<FArchive, ESPMode::ThreadSafe>& Archive, const FString& OriginalUrl, const IMediaOptions* Options) override;
+	virtual void TickInput(FTimespan DeltaTime, FTimespan Timecode) override;
+	virtual bool FlushOnSeekStarted() const;
+	virtual bool FlushOnSeekCompleted() const;
+	virtual void ProcessVideoSamples() override;
+	virtual bool GetPlayerFeatureFlag(EFeatureFlag flag) const override;
+
+protected:
+
+	/**
+	 * Check whether this player is initialized.
+	 *
+	 * @return true if initialized, false otherwise.
+	 */
+	bool IsInitialized() const;
+
+protected:
+
+	//~ IMediaCache interface
+
+	virtual bool QueryCacheState(EMediaCacheState State, TRangeSet<FTimespan>& OutTimeRanges) const override;
+
+protected:
+
+	//~ IMediaControls interface
+
+	virtual bool CanControl(EMediaControl Control) const override;
+	virtual FTimespan GetDuration() const override;
+	virtual float GetRate() const override;
+	virtual EMediaState GetState() const override;
+	virtual EMediaStatus GetStatus() const override;
+	virtual TRangeSet<float> GetSupportedRates(EMediaRateThinning Thinning) const override;
+	virtual FTimespan GetTime() const override;
+	virtual bool IsLooping() const override;
+	virtual bool Seek(const FTimespan& Time) override;
+	virtual bool SetLooping(bool Looping) override;
+	virtual bool SetRate(float Rate) override;
+	virtual void SetBlockingPlaybackHint(bool bFacadeWillUseBlockingPlayback) override;
+
+protected:
+
+	//~ IMediaSamples interface
+
+	virtual bool FetchVideo(TRange<FTimespan> TimeRange, TSharedPtr<IMediaTextureSample, ESPMode::ThreadSafe>& OutSample) override;
+	virtual void FlushSamples() override;
+
+	virtual EFetchBestSampleResult FetchBestVideoSampleForTimeRange(const TRange<FMediaTimeStamp> & TimeRange, TSharedPtr<IMediaTextureSample, ESPMode::ThreadSafe>& OutSample, bool bReverse) override;
+	virtual bool PeekVideoSampleTime(FMediaTimeStamp & TimeStamp) override;
+
+protected:
+
+	//~ IMediaTracks interface
+
+	virtual bool GetAudioTrackFormat(int32 TrackIndex, int32 FormatIndex, FMediaAudioTrackFormat& OutFormat) const override;
+	virtual int32 GetNumTracks(EMediaTrackType TrackType) const override;
+	virtual int32 GetNumTrackFormats(EMediaTrackType TrackType, int32 TrackIndex) const override;
+	virtual int32 GetSelectedTrack(EMediaTrackType TrackType) const override;
+	virtual FText GetTrackDisplayName(EMediaTrackType TrackType, int32 TrackIndex) const override;
+	virtual int32 GetTrackFormat(EMediaTrackType TrackType, int32 TrackIndex) const override;
+	virtual FString GetTrackLanguage(EMediaTrackType TrackType, int32 TrackIndex) const override;
+	virtual FString GetTrackName(EMediaTrackType TrackType, int32 TrackIndex) const override;
+	virtual bool GetVideoTrackFormat(int32 TrackIndex, int32 FormatIndex, FMediaVideoTrackFormat& OutFormat) const override;
+	virtual bool SelectTrack(EMediaTrackType TrackType, int32 TrackIndex) override;
+	virtual bool SetTrackFormat(EMediaTrackType TrackType, int32 TrackIndex, int32 FormatIndex) override;
+
+private:
+
+	/** The duration of the currently loaded media. */
+	FTimespan CurrentDuration;
+
+	/** The current playback rate. */
+	float CurrentRate;
+
+	/** The player's current state. */
+	EMediaState CurrentState;
+
+	/** The current time of the playback. */
+	FTimespan CurrentTime;
+
+	/** The URL of the currently opened media. */
+	FString CurrentUrl;
+
+	/** Whether an offset to delta time has been applied yet. */
+	bool DeltaTimeHackApplied;
+
+	/** The media event handler. */
+	IMediaEventSink& EventSink;
+
+	/** Sample time of the last fetched video sample. */
+	FTimespan LastFetchTime;
+
+	/** The image sequence loader. */
+	TSharedPtr<FImgMediaLoader, ESPMode::ThreadSafe> Loader;
+
+	/** If playback just restarted from the Stopped state. */
+	bool PlaybackRestarted;
+
+	/** The scheduler for image loading. */
+	TSharedPtr<FImgMediaScheduler, ESPMode::ThreadSafe> Scheduler;
+
+	/** Index of the selected video track. */
+	int32 SelectedVideoTrack;
+
+	/** Should the video loop to the beginning at completion */
+    bool ShouldLoop;
+
+	/** The global cache to use. */
+	TSharedPtr<FImgMediaGlobalCache, ESPMode::ThreadSafe> GlobalCache;
+
+	/** True if we have run RequestFrame already for this frame. */
+	bool RequestFrameHasRun;
+
+	/** True if facade has signaled it uses blocking playback, false if not */
+	bool PlaybackIsBlocking;
+};
diff --git a/Engine/Plugins/Media/ImgMedia/Source/ImgMedia/Private/Assets/ImgMediaMipMapInfo.cpp b/Engine/Plugins/Media/ImgMedia/Source/ImgMedia/Private/Assets/ImgMediaMipMapInfo.cpp
--- a/Engine/Plugins/Media/ImgMedia/Source/ImgMedia/Private/Assets/ImgMediaMipMapInfo.cpp
+++ b/Engine/Plugins/Media/ImgMedia/Source/ImgMedia/Private/Assets/ImgMediaMipMapInfo.cpp
@@ -1,6 +1,6 @@
 // Copyright Epic Games, Inc. All Rights Reserved.
 
-#include "ImgMediaMipMapInfo.h"
+#include "Assets/ImgMediaMipMapInfo.h"
 
 #include "ConvexVolume.h"
 #include "IImgMediaModule.h"
diff --git a/Engine/Plugins/Media/ImgMedia/Source/ImgMedia/Private/Assets/ImgMediaMipMapInfo.h b/Engine/Plugins/Media/ImgMedia/Source/ImgMedia/Private/Assets/ImgMediaMipMapInfo.h
deleted file mode 100644
diff --git a/Engine/Plugins/Media/ImgMedia/Source/ImgMedia/Private/Assets/ImgMediaSource.cpp b/Engine/Plugins/Media/ImgMedia/Source/ImgMedia/Private/Assets/ImgMediaSource.cpp
--- a/Engine/Plugins/Media/ImgMedia/Source/ImgMedia/Private/Assets/ImgMediaSource.cpp
+++ b/Engine/Plugins/Media/ImgMedia/Source/ImgMedia/Private/Assets/ImgMediaSource.cpp
@@ -4,7 +4,7 @@
 
 #include "IImgMediaModule.h"
 #include "ImgMediaGlobalCache.h"
-#include "ImgMediaMipMapInfo.h"
+#include "Assets/ImgMediaMipMapInfo.h"
 #include "ImgMediaPrivate.h"
 
 #include "HAL/FileManager.h"
diff --git a/Engine/Plugins/Media/ImgMedia/Source/ImgMedia/Private/Loader/ImgMediaLoader.cpp b/Engine/Plugins/Media/ImgMedia/Source/ImgMedia/Private/Loader/ImgMediaLoader.cpp
--- a/Engine/Plugins/Media/ImgMedia/Source/ImgMedia/Private/Loader/ImgMediaLoader.cpp
+++ b/Engine/Plugins/Media/ImgMedia/Source/ImgMedia/Private/Loader/ImgMediaLoader.cpp
@@ -1,6 +1,6 @@
 // Copyright Epic Games, Inc. All Rights Reserved.
 
-#include "ImgMediaLoader.h"
+#include "Loader/ImgMediaLoader.h"
 #include "ImgMediaGlobalCache.h"
 #include "ImgMediaPrivate.h"
 
diff --git a/Engine/Plugins/Media/ImgMedia/Source/ImgMedia/Private/Loader/ImgMediaLoader.h b/Engine/Plugins/Media/ImgMedia/Source/ImgMedia/Private/Loader/ImgMediaLoader.h
deleted file mode 100644
diff --git a/Engine/Plugins/Media/ImgMedia/Source/ImgMedia/Private/Loader/ImgMediaLoaderWork.cpp b/Engine/Plugins/Media/ImgMedia/Source/ImgMedia/Private/Loader/ImgMediaLoaderWork.cpp
--- a/Engine/Plugins/Media/ImgMedia/Source/ImgMedia/Private/Loader/ImgMediaLoaderWork.cpp
+++ b/Engine/Plugins/Media/ImgMedia/Source/ImgMedia/Private/Loader/ImgMediaLoaderWork.cpp
@@ -1,12 +1,12 @@
 // Copyright Epic Games, Inc. All Rights Reserved.
 
-#include "ImgMediaLoaderWork.h"
+#include "Loader/ImgMediaLoaderWork.h"
 #include "ImgMediaPrivate.h"
 
 #include "Misc/ScopeLock.h"
 #include "IImgMediaModule.h"
 #include "Readers/IImgMediaReader.h"
-#include "ImgMediaLoader.h"
+#include "Loader/ImgMediaLoader.h"
 
 
 /** Time spent abandoning worker threads. */
diff --git a/Engine/Plugins/Media/ImgMedia/Source/ImgMedia/Private/Player/ImgMediaPlayer.cpp b/Engine/Plugins/Media/ImgMedia/Source/ImgMedia/Private/Player/ImgMediaPlayer.cpp
--- a/Engine/Plugins/Media/ImgMedia/Source/ImgMedia/Private/Player/ImgMediaPlayer.cpp
+++ b/Engine/Plugins/Media/ImgMedia/Source/ImgMedia/Private/Player/ImgMediaPlayer.cpp
@@ -1,6 +1,6 @@
 // Copyright Epic Games, Inc. All Rights Reserved.
 
-#include "ImgMediaPlayer.h"
+#include "Player/ImgMediaPlayer.h"
 #include "ImgMediaPrivate.h"
 #include "ImgMediaSource.h"
 
diff --git a/Engine/Plugins/Media/ImgMedia/Source/ImgMedia/Private/Player/ImgMediaPlayer.h b/Engine/Plugins/Media/ImgMedia/Source/ImgMedia/Private/Player/ImgMediaPlayer.h
deleted file mode 100644
diff --git a/Engine/Plugins/Media/ImgMedia/Source/ImgMediaEditor/ImgMediaEditor.Build.cs b/Engine/Plugins/Media/ImgMedia/Source/ImgMediaEditor/ImgMediaEditor.Build.cs
--- a/Engine/Plugins/Media/ImgMedia/Source/ImgMediaEditor/ImgMediaEditor.Build.cs
+++ b/Engine/Plugins/Media/ImgMedia/Source/ImgMediaEditor/ImgMediaEditor.Build.cs
@@ -30,7 +30,7 @@ namespace UnrealBuildTool.Rules
 
 			PrivateIncludePaths.AddRange(
 				new string[] {
-					"ImgMedia/Private",
+					//"ImgMedia/Private",
 				});
 
 			// Are we using the engine?
