ImgMedia: Fix image sequence tile and mip support for scene captures.

diff --git a/Engine/Plugins/Media/ImgMedia/Source/ImgMedia/Private/ImgMediaSceneViewExtension.cpp b/Engine/Plugins/Media/ImgMedia/Source/ImgMedia/Private/ImgMediaSceneViewExtension.cpp
--- a/Engine/Plugins/Media/ImgMedia/Source/ImgMedia/Private/ImgMediaSceneViewExtension.cpp
+++ b/Engine/Plugins/Media/ImgMedia/Source/ImgMedia/Private/ImgMediaSceneViewExtension.cpp
@@ -24,6 +24,20 @@ void FImgMediaSceneViewExtension::SetupViewFamily(FSceneViewFamily& InViewFamily
 
 void FImgMediaSceneViewExtension::SetupView(FSceneViewFamily& InViewFamily, FSceneView& InView)
 {
+	TRACE_CPUPROFILER_EVENT_SCOPE(FImgMediaSceneViewExtension::SetupView);
+
+	/**
+	* NOTE: Scene captures call `SetupView` after the primary `BeginRenderViewFamily` call:
+	* FRendererModule::BeginRenderingViewFamilies
+	*     -> USceneCaptureComponent::UpdateDeferredCaptures
+	*         -> FScene::UpdateSceneCaptureContents
+	*             -> ISceneViewExtension::SetupView
+	* Therefore, the view infos we cache here will be correctly kept until the next frame.
+	*/
+	if (InView.bIsSceneCapture)
+	{
+		CacheViewInfo(InViewFamily, InView);
+	}
 }
 
 void FImgMediaSceneViewExtension::BeginRenderViewFamily(FSceneViewFamily& InViewFamily)
@@ -36,6 +50,28 @@ void FImgMediaSceneViewExtension::BeginRenderViewFamily(FSceneViewFamily& InView
 		LastFrameNumber = InViewFamily.FrameNumber;
 	}
 
+	for (const FSceneView* View : InViewFamily.Views)
+	{
+		/** NOTE: Scene captures currently don't call BeginRenderViewFamily() so we handle them separately in SetupView. */
+		if (View && !View->bIsSceneCapture)
+		{
+			CacheViewInfo(InViewFamily, *View);
+		}
+	}
+}
+
+int32 FImgMediaSceneViewExtension::GetPriority() const
+{
+	// Lowest priority value to ensure all other extensions are executed before ours.
+	return MIN_int32;
+}
+
+void FImgMediaSceneViewExtension::CacheViewInfo(FSceneViewFamily& InViewFamily, const FSceneView& View)
+{
+	static const auto CVarMinAutomaticViewMipBiasOffset = IConsoleManager::Get().FindTConsoleVariableDataFloat(TEXT("r.ViewTextureMipBias.Offset"));
+	static const auto CVarMinAutomaticViewMipBias = IConsoleManager::Get().FindTConsoleVariableDataFloat(TEXT("r.ViewTextureMipBias.Min"));
+	const float FieldOfViewMultiplier = CVarImgMediaFieldOfViewMultiplier.GetValueOnGameThread();
+
 	float ResolutionFraction = InViewFamily.SecondaryViewFraction;
 
 	if (InViewFamily.GetScreenPercentageInterface())
@@ -44,82 +80,69 @@ void FImgMediaSceneViewExtension::BeginRenderViewFamily(FSceneViewFamily& InView
 		ResolutionFraction *= UpperBounds[GDynamicPrimaryResolutionFraction];
 	}
 
-	static const auto CVarMinAutomaticViewMipBiasOffset = IConsoleManager::Get().FindTConsoleVariableDataFloat(TEXT("r.ViewTextureMipBias.Offset"));
-	static const auto CVarMinAutomaticViewMipBias = IConsoleManager::Get().FindTConsoleVariableDataFloat(TEXT("r.ViewTextureMipBias.Min"));
-	const float FieldOfViewMultiplier = CVarImgMediaFieldOfViewMultiplier.GetValueOnGameThread();
+	FImgMediaViewInfo Info;
+	Info.Location = View.ViewMatrices.GetViewOrigin();
+	Info.ViewDirection = View.GetViewDirection();
+	Info.ViewProjectionMatrix = View.ViewMatrices.GetViewProjectionMatrix();
 
-	for (const FSceneView* View : InViewFamily.Views)
+	if (FMath::IsNearlyEqual(FieldOfViewMultiplier, 1.0f))
+	{
+		Info.OverscanViewProjectionMatrix = Info.ViewProjectionMatrix;
+	}
+	else
 	{
-		FImgMediaViewInfo Info;
-		Info.Location = View->ViewMatrices.GetViewOrigin();
-		Info.ViewDirection = View->GetViewDirection();
-		Info.ViewProjectionMatrix = View->ViewMatrices.GetViewProjectionMatrix();
+		FMatrix AdjustedProjectionMatrix = View.ViewMatrices.GetProjectionMatrix();
 
-		if (FMath::IsNearlyEqual(FieldOfViewMultiplier, 1.0f))
-		{
-			Info.OverscanViewProjectionMatrix = Info.ViewProjectionMatrix;
-		}
-		else
-		{
-			FMatrix AdjustedProjectionMatrix = View->ViewMatrices.GetProjectionMatrix();
+		const double HalfHorizontalFOV = FMath::Atan(1.0 / AdjustedProjectionMatrix.M[0][0]);
+		const double HalfVerticalFOV = FMath::Atan(1.0 / AdjustedProjectionMatrix.M[1][1]);
 
-			const double HalfHorizontalFOV = FMath::Atan(1.0 / AdjustedProjectionMatrix.M[0][0]);
-			const double HalfVerticalFOV = FMath::Atan(1.0 / AdjustedProjectionMatrix.M[1][1]);
+		AdjustedProjectionMatrix.M[0][0] = 1.0 / FMath::Tan(HalfHorizontalFOV * FieldOfViewMultiplier);
+		AdjustedProjectionMatrix.M[1][1] = 1.0 / FMath::Tan(HalfVerticalFOV * FieldOfViewMultiplier);
 
-			AdjustedProjectionMatrix.M[0][0] = 1.0 / FMath::Tan(HalfHorizontalFOV * FieldOfViewMultiplier);
-			AdjustedProjectionMatrix.M[1][1] = 1.0 / FMath::Tan(HalfVerticalFOV * FieldOfViewMultiplier);
-			
-			Info.OverscanViewProjectionMatrix = View->ViewMatrices.GetViewMatrix() * AdjustedProjectionMatrix;
-		}
-		
-		Info.ViewportRect = View->UnconstrainedViewRect.Scale(ResolutionFraction);
+		Info.OverscanViewProjectionMatrix = View.ViewMatrices.GetViewMatrix() * AdjustedProjectionMatrix;
+	}
 
-		// We store hidden or show-only ids to later avoid needless calculations when objects are not in view.
-		if (View->ShowOnlyPrimitives.IsSet())
-		{
-			Info.bPrimitiveHiddenMode = false;
-			Info.PrimitiveComponentIds = View->ShowOnlyPrimitives.GetValue();
-		}
-		else
-		{
-			Info.bPrimitiveHiddenMode = true;
-			Info.PrimitiveComponentIds = View->HiddenPrimitives;
-		}
+	Info.ViewportRect = View.UnconstrainedViewRect.Scale(ResolutionFraction);
 
-		/* View->MaterialTextureMipBias is only set later in rendering so we replicate here the calculations
-		 * found in FSceneRenderer::PreVisibilityFrameSetup.*/
-		if (View->PrimaryScreenPercentageMethod == EPrimaryScreenPercentageMethod::TemporalUpscale)
-		{
-			const float EffectivePrimaryResolutionFraction = float(Info.ViewportRect.Width()) / (View->UnscaledViewRect.Width() * InViewFamily.SecondaryViewFraction);
-			Info.MaterialTextureMipBias = -(FMath::Max(-FMath::Log2(EffectivePrimaryResolutionFraction), 0.0f)) + CVarMinAutomaticViewMipBiasOffset->GetValueOnGameThread();
-			Info.MaterialTextureMipBias = FMath::Max(Info.MaterialTextureMipBias, CVarMinAutomaticViewMipBias->GetValueOnGameThread());
-
-			if (!ensureMsgf(!FMath::IsNaN(Info.MaterialTextureMipBias) && FMath::IsFinite(Info.MaterialTextureMipBias), TEXT("Calculated material texture mip bias is invalid, defaulting to zero.")))
-			{
-				Info.MaterialTextureMipBias = 0.0f;
-			}
-		}
-		else
+	// We store hidden or show-only ids to later avoid needless calculations when objects are not in view.
+	if (View.ShowOnlyPrimitives.IsSet())
+	{
+		Info.bPrimitiveHiddenMode = false;
+		Info.PrimitiveComponentIds = View.ShowOnlyPrimitives.GetValue();
+	}
+	else
+	{
+		Info.bPrimitiveHiddenMode = true;
+		Info.PrimitiveComponentIds = View.HiddenPrimitives;
+	}
+
+	/* View.MaterialTextureMipBias is only set later in rendering so we replicate here the calculations
+	 * found in FSceneRenderer::PreVisibilityFrameSetup.*/
+	if (View.PrimaryScreenPercentageMethod == EPrimaryScreenPercentageMethod::TemporalUpscale)
+	{
+		const float EffectivePrimaryResolutionFraction = float(Info.ViewportRect.Width()) / (View.UnscaledViewRect.Width() * InViewFamily.SecondaryViewFraction);
+		Info.MaterialTextureMipBias = -(FMath::Max(-FMath::Log2(EffectivePrimaryResolutionFraction), 0.0f)) + CVarMinAutomaticViewMipBiasOffset->GetValueOnGameThread();
+		Info.MaterialTextureMipBias = FMath::Max(Info.MaterialTextureMipBias, CVarMinAutomaticViewMipBias->GetValueOnGameThread());
+
+		if (!ensureMsgf(!FMath::IsNaN(Info.MaterialTextureMipBias) && FMath::IsFinite(Info.MaterialTextureMipBias), TEXT("Calculated material texture mip bias is invalid, defaulting to zero.")))
 		{
 			Info.MaterialTextureMipBias = 0.0f;
 		}
+	}
+	else
+	{
+		Info.MaterialTextureMipBias = 0.0f;
+	}
 
 #if !(UE_BUILD_SHIPPING || UE_BUILD_TEST)
-		static const auto CVarMipMapDebug = IConsoleManager::Get().FindConsoleVariable(TEXT("ImgMedia.MipMapDebug"));
-
-		if (GEngine != nullptr && CVarMipMapDebug != nullptr && CVarMipMapDebug->GetBool())
-		{
-			const FString ViewName = InViewFamily.ProfileDescription.IsEmpty() ? TEXT("View"): InViewFamily.ProfileDescription;
-			GEngine->AddOnScreenDebugMessage(-1, 0.0f, FColor::Cyan, *FString::Printf(TEXT("%s location: [%s], direction: [%s]"), *ViewName, *Info.Location.ToString(), *View->GetViewDirection().ToString()));
-		}
-#endif
+	static const auto CVarMipMapDebug = IConsoleManager::Get().FindConsoleVariable(TEXT("ImgMedia.MipMapDebug"));
 
-		CachedViewInfos.Add(MoveTemp(Info));
+	if (GEngine != nullptr && CVarMipMapDebug != nullptr && CVarMipMapDebug->GetBool())
+	{
+		const FString ViewName = InViewFamily.ProfileDescription.IsEmpty() ? TEXT("View") : InViewFamily.ProfileDescription;
+		GEngine->AddOnScreenDebugMessage(-1, 0.0f, FColor::Cyan, *FString::Printf(TEXT("%s location: [%s], direction: [%s]"), *ViewName, *Info.Location.ToString(), *View.GetViewDirection().ToString()));
 	}
-}
+#endif
 
-int32 FImgMediaSceneViewExtension::GetPriority() const
-{
-	// Lowest priority value to ensure all other extensions are executed before ours.
-	return MIN_int32;
+	CachedViewInfos.Add(MoveTemp(Info));
 }
diff --git a/Engine/Plugins/Media/ImgMedia/Source/ImgMedia/Public/ImgMediaSceneViewExtension.h b/Engine/Plugins/Media/ImgMedia/Source/ImgMedia/Public/ImgMediaSceneViewExtension.h
--- a/Engine/Plugins/Media/ImgMedia/Source/ImgMedia/Public/ImgMediaSceneViewExtension.h
+++ b/Engine/Plugins/Media/ImgMedia/Source/ImgMedia/Public/ImgMediaSceneViewExtension.h
@@ -51,6 +51,9 @@ public:
 	int32 GetPriority() const override;
 
 private:
+	/** Cache camera view information for the current frame. */
+	void CacheViewInfo(FSceneViewFamily& InViewFamily, const FSceneView& View);
+	
 	/** Array of info on each camera used for mipmap calculations. */
 	TArray<FImgMediaViewInfo> CachedViewInfos;
 
