Fixing issue with very short image sequences leading to timeout scenarious during blocked playback
#rnx

diff --git a/Engine/Plugins/Media/ImgMedia/Source/ImgMedia/Internal/Loader/ImgMediaLoader.h b/Engine/Plugins/Media/ImgMedia/Source/ImgMedia/Internal/Loader/ImgMediaLoader.h
--- a/Engine/Plugins/Media/ImgMedia/Source/ImgMedia/Internal/Loader/ImgMediaLoader.h
+++ b/Engine/Plugins/Media/ImgMedia/Source/ImgMedia/Internal/Loader/ImgMediaLoader.h
@@ -152,7 +152,7 @@ public:
 	/**
 	 * Notify loader out a seek
 	 */
-	void Seek(const FMediaTimeStamp& SeekTarget);
+	void Seek(const FMediaTimeStamp& SeekTarget, bool bReverse);
 
 	/**
 	 * Tries to get the best sample for a given time range.
@@ -510,27 +510,6 @@ protected:
 	 */
 	FTimespan ModuloTime(FTimespan Time);
 
-	/**
-	 * Gets the amount of overlap (in seconds) between a frame and a time range.
-	 * A negative or zero  value indicates the frame does not overlap the range.
-	 *
-	 * @param FrameIndex Index of frame to check.
-	 * @param StartTime Start of the range.
-	 * @param EndTime End of the range.
-	 * @param Amount of overlap.
-	 */
-	float GetFrameOverlap(uint32 FrameIndex, FTimespan StartTime, FTimespan EndTime) const;
-
-	/**
-	 * Find maximum overlapping frame index for given range
-	 */
-	float FindMaxOverlapInRange(int32 StartIndex, int32 EndIndex, FTimespan StartTime, FTimespan EndTime, int32& MaxIdx) const;
-
-	/**
-	 * Get frame data for given index. If not available attempt to find earlier frame up to given index
-	 */
-	const TSharedPtr<FImgMediaFrame, ESPMode::ThreadSafe>* GetFrameForBestIndex(int32 & MaxIdx, int32 LastIndex);
-
 	/**
 	 * Helper function to get the number at the end of a string.
 	 * 
@@ -630,6 +609,7 @@ private:
 	FName SequenceName;
 
 private:
+	bool RequestFrame(int32 FrameNumber, float PlayRate, bool Loop);
 
 	/** Index of the previously requested frame. */
 	int32 LastRequestedFrame;
diff --git a/Engine/Plugins/Media/ImgMedia/Source/ImgMedia/Private/Loader/ImgMediaLoader.cpp b/Engine/Plugins/Media/ImgMedia/Source/ImgMedia/Private/Loader/ImgMediaLoader.cpp
--- a/Engine/Plugins/Media/ImgMedia/Source/ImgMedia/Private/Loader/ImgMediaLoader.cpp
+++ b/Engine/Plugins/Media/ImgMedia/Source/ImgMedia/Private/Loader/ImgMediaLoader.cpp
@@ -28,9 +28,6 @@
 	#include "Readers/ExrImgMediaReader.h"
 #endif
 
-#define ALLOW_CHEATED_FRAMES 0		// if set to 1 the loader is allowed to return "close by" frame data with a fake PTS to optimize for quick delivery IF there is no blocked range is set
-									// ^^^ DISABLED as this can lead to entirely incorrect frames being displayed when seeking in paused state
-
 /** Time spent loading a new image sequence. */
 DECLARE_CYCLE_STAT(TEXT("ImgMedia Loader Load Sequence"), STAT_ImgMedia_LoaderLoadSequence, STATGROUP_Media);
 
@@ -346,6 +343,7 @@ void FImgMediaLoader::ResetFetchLogic()
 {
 	QueuedSampleFetch.LastFrameIndex = INDEX_NONE;
 	QueuedSampleFetch.LastTimeStamp.Invalidate();
+	QueuedSampleFetch.LastDuration = FTimespan::Zero();
 	QueuedSampleFetch.LoopIndex = 0;
 
 	SkipFramesCounter = 0;
@@ -365,89 +363,50 @@ void FImgMediaLoader::HandlePause()
 #endif
 }
 
-float FImgMediaLoader::FindMaxOverlapInRange(int32 StartIndex, int32 EndIndex, FTimespan StartTime, FTimespan EndTime, int32 & MaxIdx) const
+
+void FImgMediaLoader::Seek(const FMediaTimeStamp& SeekTarget, bool bReverse)
 {
-	int32 IdxInc = 0;
-	if (StartIndex < EndIndex)
+	ResetFetchLogic();
+	if (!bReverse)
 	{
-		IdxInc = 1;
+		QueuedSampleFetch.LastTimeStamp = SeekTarget;
 	}
 	else
 	{
-		IdxInc = -1;
-	}
-
-	// Find index that overlaps the most and is furthest along the timeline...
-	MaxIdx = -1;
-	float MaxOverlap = 0.0f;
-	for (uint32 Idx = StartIndex; Idx != (EndIndex + IdxInc); Idx+=IdxInc)
-	{
-		float Overlap = GetFrameOverlap(Idx, StartTime, EndTime);
-		if (MaxOverlap < Overlap)
-		{
-			MaxOverlap = Overlap;
-			MaxIdx = Idx;
-		}
-	}
-	return MaxOverlap;
-}
-
-
-const TSharedPtr<FImgMediaFrame, ESPMode::ThreadSafe>* FImgMediaLoader::GetFrameForBestIndex(int32 & MaxIdx, int32 LastIndex)
-{
-	const TSharedPtr<FImgMediaFrame, ESPMode::ThreadSafe>* Frame = nullptr;
-	int32 IdxInc = (MaxIdx > LastIndex) ? -1 : 1;
-
-	while (MaxIdx != (LastIndex + IdxInc))
-	{
-		Frame = UseGlobalCache ? GlobalCache->FindAndTouch(SequenceName, MaxIdx) : Frames.FindAndTouch(MaxIdx);
-		if (Frame)
-		{
-			break;
-		}
-		MaxIdx += IdxInc;
+		QueuedSampleFetch.LastTimeStamp = SeekTarget - FTimespan::FromSeconds(static_cast<double>(SequenceFrameRate.Denominator) / SequenceFrameRate.Numerator);
 	}
-	return Frame;
-}
-
-
-void FImgMediaLoader::Seek(const FMediaTimeStamp& SeekTarget)
-{
-	ResetFetchLogic();
-	QueuedSampleFetch.LastTimeStamp = SeekTarget;
-	QueuedSampleFetch.LastDuration = FTimespan::Zero();
-	QueuedSampleFetch.LoopIndex = 0;
-	QueuedSampleFetch.LastFrameIndex = INDEX_NONE;
-
-	QueuedSampleFetch.LastTimeStamp.Time = FrameNumberToTime(TimeToFrameNumber(QueuedSampleFetch.LastTimeStamp.Time));
 }
 
 
 bool FImgMediaLoader::DiscardVideoSamples(const TRange<FMediaTimeStamp>& TimeRange, bool bIsLoopingEnabled, float PlayRate)
 {
 	bool bOk = false;
-	if (PlayRate >= 0.0f)
+
+	if (QueuedSampleFetch.LastTimeStamp.IsValid())
 	{
-		if ((QueuedSampleFetch.LastTimeStamp + QueuedSampleFetch.LastDuration) < TimeRange.GetUpperBoundValue())
+		if (PlayRate >= 0.0f)
 		{
-			QueuedSampleFetch.LastTimeStamp = TimeRange.GetUpperBoundValue();
-			bOk = true;
+			if ((QueuedSampleFetch.LastTimeStamp + QueuedSampleFetch.LastDuration) < TimeRange.GetUpperBoundValue())
+			{
+				QueuedSampleFetch.LastTimeStamp = TimeRange.GetUpperBoundValue();
+				bOk = true;
+			}
 		}
-	}
-	else
-	{
-		if ((QueuedSampleFetch.LastTimeStamp - QueuedSampleFetch.LastDuration) > TimeRange.GetLowerBoundValue())
+		else
 		{
-			QueuedSampleFetch.LastTimeStamp = TimeRange.GetLowerBoundValue();
-			bOk = true;
+			if ((QueuedSampleFetch.LastTimeStamp - QueuedSampleFetch.LastDuration) > TimeRange.GetLowerBoundValue())
+			{
+				QueuedSampleFetch.LastTimeStamp = TimeRange.GetLowerBoundValue();
+				bOk = true;
+			}
 		}
-	}
 
-	if (bOk)
-	{
-		QueuedSampleFetch.LastDuration = FTimespan::Zero();
-		QueuedSampleFetch.LastTimeStamp.Time = FrameNumberToTime(TimeToFrameNumber(QueuedSampleFetch.LastTimeStamp.Time));
-		QueuedSampleFetch.LastFrameIndex = INDEX_NONE;
+		if (bOk)
+		{
+			QueuedSampleFetch.LastDuration = FTimespan::Zero();
+			QueuedSampleFetch.LastTimeStamp.Time = FrameNumberToTime(TimeToFrameNumber(QueuedSampleFetch.LastTimeStamp.Time));
+			QueuedSampleFetch.LastFrameIndex = INDEX_NONE;
+		}
 	}
 
 	return bOk;
@@ -469,7 +428,6 @@ IMediaSamples::EFetchBestSampleResult FImgMediaLoader::FetchBestVideoSampleForTi
 		{
 			if (PlayRate >= 0.0f)
 			{
-//WHAT IF THE LAST FRAME WAS SHOWN IN PAUSE MODE?
 				FMediaTimeStamp NextTimeStamp = QueuedSampleFetch.LastTimeStamp + QueuedSampleFetch.LastDuration;
 				if (NextTimeStamp >= TimeRange.GetUpperBoundValue())
 					{
@@ -499,7 +457,18 @@ IMediaSamples::EFetchBestSampleResult FImgMediaLoader::FetchBestVideoSampleForTi
 		FTimespan StartTime = TimeRange.GetLowerBoundValue().Time;
 		int32 StartLoopIndex = FMediaTimeStamp::GetSecondaryIndex(TimeRange.GetLowerBoundValue().SequenceIndex);
 		FTimespan EndTime = TimeRange.GetUpperBoundValue().Time;
-		int32 EndLoopIndex = FMediaTimeStamp::GetSecondaryIndex(TimeRange.GetUpperBoundValue().SequenceIndex);
+		int32 EndLoopIndex;
+		// The end time is EXCLUSIVE, move the time by the smallest amount we can so we correctly select indices of frames
+		if (EndTime > FTimespan::Zero())
+		{
+			EndTime -= FTimespan(1);
+			EndLoopIndex = FMediaTimeStamp::GetSecondaryIndex(TimeRange.GetUpperBoundValue().SequenceIndex);
+		}
+		else
+		{
+			EndTime = SequenceDuration - FTimespan(1);
+			EndLoopIndex = FMediaTimeStamp::GetSecondaryIndex(TimeRange.GetUpperBoundValue().SequenceIndex) - 1;
+		}
 
 		if (bIsLoopingEnabled)
 		{
@@ -529,72 +498,87 @@ IMediaSamples::EFetchBestSampleResult FImgMediaLoader::FetchBestVideoSampleForTi
 
 		// Find the frame that overlaps the most with the given range & is furthest along on the timeline
 		int32 MaxIdx = -1;
+		int32 MaxIdxLoopIndexOffset = 0;
 		if (PlayRate >= 0.0f)
 		{
 			// Forward...
 
-			// Is the start index the same as the end index?
-			if (StartIndex == EndIndex)
+			// Per default we select the first frame in range
+			MaxIdx = StartIndex;
+
+			// More than one possibility?
+			if (EndIndex != StartIndex || StartLoopIndex != EndLoopIndex)
 			{
-				// Is StartTime less than the time of StartIndex?
-				FTimespan CurrentStartTime = FrameNumberToTime(StartIndex);
-				if (StartTime < CurrentStartTime)
+				FTimespan FrameLength = FTimespan::FromSeconds(static_cast<double>(SequenceFrameRate.Denominator) / SequenceFrameRate.Numerator);
+
+				// Compute visible duration of first frame...
+				FTimespan FirstTime = FrameNumberToTime(StartIndex);
+				FTimespan FirstFrameDuration = FrameLength - FTimespan(FMath::Max((TimeRange.GetLowerBoundValue().Time - FirstTime).GetTicks(), 0));
+
+				// Compute the length of the last frame
+				FTimespan LastTime = FrameNumberToTime(EndIndex);
+				FTimespan LastFrameDuration = FrameLength - FTimespan(FMath::Max((LastTime - TimeRange.GetUpperBoundValue().Time).GetTicks(), 0));
+
+				// First one wins by default
+				FTimespan MaxFrameDuration = FirstFrameDuration;
+
+				int32 NumFrames = GetNumImages();
+
+				// Do we have a "second to last" frame? (it would always have full coverage)
+				int64 EndIndexUnlooped = (EndLoopIndex - StartLoopIndex) * NumFrames + EndIndex;
+				if ((EndIndexUnlooped - StartIndex) > 1)
 				{
-					// Yes. TimeToFrameNumber rounded up, but we really don't want that
-					// so adjust StartIndex.
-					if (StartIndex > 0)
-					{
-						StartIndex--;
-					}
+					// We do. Record it as best so far instead of the first one...
+					MaxIdx = (EndIndex > 0) ? (EndIndex - 1) : (NumFrames - 1);
+					MaxFrameDuration = FrameLength;
 				}
-			}
 
-			// Does the range span across the end of the media?
-			if ((StartIndex > EndIndex) || (EndLoopIndex > StartLoopIndex))
-			{
-				// Yes...
-				int32 MaxIdx1, MaxIdx2;
-				float MaxOverlap1 = FindMaxOverlapInRange(StartIndex, GetNumImages()-1, StartTime, FrameNumberToTime(GetNumImages()), MaxIdx1);
-				float MaxOverlap2 = FindMaxOverlapInRange(0, EndIndex, FTimespan::Zero(), EndTime, MaxIdx2);
-				MaxIdx = (MaxOverlap2 >= MaxOverlap1) ? MaxIdx2 : MaxIdx1;
-			}
-			else
-			{
-				FindMaxOverlapInRange(StartIndex, EndIndex, StartTime, EndTime, MaxIdx);
+				// Is the last one even better?
+				if (LastFrameDuration >= MaxFrameDuration)
+				{
+					MaxIdx = EndIndex;
+				}
 			}
 		}
 		else
 		{
-			// Backward...
-			
-			// Is the start index the same as the end index?
-			if (StartIndex == EndIndex)
+			// Backwards...
+
+			// Per default we select the last frame in range
+			MaxIdx = EndIndex;
+
+			// More than one possibility?
+			if (EndIndex != StartIndex || StartLoopIndex != EndLoopIndex)
 			{
-				// Is StartTime larger than the time of StartIndex?
-				FTimespan CurrentEndTime = FrameNumberToTime(EndIndex);
-				if (EndTime > CurrentEndTime)
+				FTimespan FrameLength = FTimespan::FromSeconds(static_cast<double>(SequenceFrameRate.Denominator) / SequenceFrameRate.Numerator);
+
+				// Compute the length of the last frame
+				FTimespan FirstTime = FrameNumberToTime(EndIndex);
+				FTimespan FirstFrameDuration = FrameLength - FTimespan(FMath::Max((FirstTime - TimeRange.GetUpperBoundValue().Time).GetTicks(), 0));
+
+				// Compute visible duration of first frame...
+				FTimespan LastTime = FrameNumberToTime(StartIndex);
+				FTimespan LastFrameDuration = FrameLength - FTimespan(FMath::Max((TimeRange.GetLowerBoundValue().Time - LastTime).GetTicks(), 0));
+
+				// First one wins by default
+				FTimespan MaxFrameDuration = FirstFrameDuration;
+
+				int32 NumFrames = GetNumImages();
+
+				// Do we have a "second to last" frame? (it would always have full coverage)
+				int64 EndIndexUnlooped = (EndLoopIndex - StartLoopIndex) * NumFrames + EndIndex;
+				if ((EndIndexUnlooped - StartIndex) > 1)
 				{
-					// Yes. TimeToFrameNumber rounded up, but we really don't want that
-					// so adjust StartIndex.
-					if (EndIndex < (GetNumImages() - 1))
-					{
-						++EndIndex;
-					}
+					// We do. Record it as best so far instead of the first one...
+					MaxIdx = (StartIndex < (NumFrames - 1)) ? (StartIndex + 1) : 0;
+					MaxFrameDuration = FrameLength;
 				}
-			}
 
-			// Does the range span across the end of the media?
-			if ((StartIndex > EndIndex) || (StartLoopIndex < EndLoopIndex))
-			{
-				// Yes...
-				int32 MaxIdx1, MaxIdx2;
-				float MaxOverlap1 = FindMaxOverlapInRange(EndIndex, 0, FTimespan::Zero(), EndTime, MaxIdx1);
-				float MaxOverlap2 = FindMaxOverlapInRange(GetNumImages()-1, StartIndex,  StartTime, FrameNumberToTime(GetNumImages()), MaxIdx2);
-				MaxIdx = (MaxOverlap2 >= MaxOverlap1) ? MaxIdx2 : MaxIdx1;
-			}
-			else
-			{
-				FindMaxOverlapInRange(EndIndex, StartIndex, StartTime, EndTime, MaxIdx);
+				// Is the last one even better?
+				if (LastFrameDuration >= MaxFrameDuration)
+				{
+					MaxIdx = StartIndex;
+				}
 			}
 		}
 
@@ -604,75 +588,12 @@ IMediaSamples::EFetchBestSampleResult FImgMediaLoader::FetchBestVideoSampleForTi
 			const TSharedPtr<FImgMediaFrame, ESPMode::ThreadSafe>* Frame;
 
 			// Request data for the frame we would like... (in case it's not in, yet)
-			RequestFrame(FrameNumberToTime(MaxIdx), PlayRate, bIsLoopingEnabled);
+			RequestFrame(MaxIdx, PlayRate, bIsLoopingEnabled);
 
 			FScopeLock Lock(&CriticalSection);
 
-#if ALLOW_CHEATED_FRAMES
-			// If playback is not blocking, we expect less expectancy of precision on the users side, but more need for speedy return of "some ok frame"
-			// So: if we detect non-blocking playback we return a "as good sample as we can", but not always the "perfect" one we calculated
-			// (still we adhere to a rough emulation of a classic output pipeline as other players have)
-			if (!bIsPlaybackBlocking)
-			{
-				// Check what data we actually have in the cache already & attempt to go further backwards on the time line
-				// to get a less then optimal frame that is still on screen and available...
-				if (PlayRate >= 0.0f)
-				{
-					// Forward...
-					if (StartIndex > EndIndex)
-					{
-						if (MaxIdx < StartIndex)
-						{
-							Frame = GetFrameForBestIndex(MaxIdx, 0);
-							if (!Frame)
-							{
-								MaxIdx = GetNumImages() - 1;
-								Frame = GetFrameForBestIndex(MaxIdx, StartIndex);
-							}
-
-						}
-						else
-						{
-							Frame = GetFrameForBestIndex(MaxIdx, StartIndex);
-						}
-					}
-					else
-					{
-						Frame = GetFrameForBestIndex(MaxIdx, StartIndex);
-					}
-				}
-				else
-				{
-					// Backward...
-					if (StartIndex > EndIndex)
-					{
-						if (MaxIdx > EndIndex)
-						{
-							Frame = GetFrameForBestIndex(MaxIdx, GetNumImages() - 1);
-							if (!Frame)
-							{
-								MaxIdx = 0;
-								Frame = GetFrameForBestIndex(MaxIdx, EndIndex);
-							}
-
-						}
-						else
-						{
-							Frame = GetFrameForBestIndex(MaxIdx, EndIndex);
-						}
-					}
-					else
-					{
-						Frame = GetFrameForBestIndex(MaxIdx, EndIndex);
-					}
-				}
-			}
-			else
-#endif
-			{
-				// Get a frame if we have one available right now...
-				Frame = UseGlobalCache ? GlobalCache->FindAndTouch(SequenceName, MaxIdx) : Frames.FindAndTouch(MaxIdx);
-			}
+			// Get a frame if we have one available right now...
+			Frame = UseGlobalCache ? GlobalCache->FindAndTouch(SequenceName, MaxIdx) : Frames.FindAndTouch(MaxIdx);
 
 			// Got a potential frame?
 			if (Frame)
@@ -687,47 +608,44 @@ IMediaSamples::EFetchBestSampleResult FImgMediaLoader::FetchBestVideoSampleForTi
 
 				RetryCount = 0;
 
+				FTimespan SampleTime = FrameNumberToTime(MaxIdx);
 
 				// Check if we wrapped around the end of the media vs. the last frame we returned...
-				if (QueuedSampleFetch.LastFrameIndex != INDEX_NONE && bIsLoopingEnabled)
+				// (skip this if we are not looping or if this is either the first frame ever or first frame after a seek)
+				if (bIsLoopingEnabled && QueuedSampleFetch.LastTimeStamp.IsValid() && QueuedSampleFetch.LastFrameIndex != INDEX_NONE)
 				{
 					if (PlayRate >= 0.0f)
 					{
-						if (QueuedSampleFetch.LastFrameIndex > MaxIdx)
+						if (QueuedSampleFetch.LastTimeStamp.Time >= SampleTime)
 						{
 							++QueuedSampleFetch.LoopIndex;
 						}
 					}
 					else
 					{
-						if (QueuedSampleFetch.LastFrameIndex < MaxIdx)
+						if (QueuedSampleFetch.LastTimeStamp.Time <= SampleTime)
 						{
 							--QueuedSampleFetch.LoopIndex;
 						}
 					}
 				}
 
-				int64 NewSequenceIndex = FMediaTimeStamp::MakeSequenceIndex(FMediaTimeStamp::GetPrimaryIndex(TimeRange.GetLowerBoundValue().SequenceIndex), QueuedSampleFetch.LoopIndex);
-				
-				// Different from the last one we returned?
-				if ((QueuedSampleFetch.LastFrameIndex != MaxIdx) || (QueuedSampleFetch.LastTimeStamp.IsValid() && (QueuedSampleFetch.LastTimeStamp.SequenceIndex != NewSequenceIndex)) || (ImagePaths.Num() == 1))
-				{
-					QueuedSampleFetch.LastFrameIndex = MaxIdx;
-					QueuedSampleFetch.LastTimeStamp = FMediaTimeStamp(FrameNumberToTime(MaxIdx), NewSequenceIndex);
-					QueuedSampleFetch.LastDuration = FTimespan::FromSeconds(Frame->Get()->Info.FrameRate.AsInterval());
+				// Track state & setup sample for caller...
+				QueuedSampleFetch.LastFrameIndex = MaxIdx;
+				QueuedSampleFetch.LastTimeStamp = FMediaTimeStamp(SampleTime, FMediaTimeStamp::MakeSequenceIndex(FMediaTimeStamp::GetPrimaryIndex(TimeRange.GetLowerBoundValue().SequenceIndex), QueuedSampleFetch.LoopIndex));
+				QueuedSampleFetch.LastDuration = FTimespan::FromSeconds(Frame->Get()->Info.FrameRate.AsInterval());
 
-					// We are clear to return it as new result... Make a sample & initialize it...
-					auto Sample = MakeShared<FImgMediaTextureSample, ESPMode::ThreadSafe>();
-					
+				// We are clear to return it as new result... Make a sample & initialize it...
+				auto Sample = MakeShared<FImgMediaTextureSample, ESPMode::ThreadSafe>();
 
-					ImgMediaLoader::CheckAndUpdateImgDimensions(SequenceDim, Frame->Get()->Info.Dim);
 
-					if (Sample->Initialize(*Frame->Get(), SequenceDim, QueuedSampleFetch.LastTimeStamp, QueuedSampleFetch.LastDuration, GetNumMipLevels(), TilingDescription))
-					{
-						OutSample = Sample;
-						CSV_EVENT(ImgMedia, TEXT("LoaderFetchHit %d %d-%d"), MaxIdx, StartIndex, EndIndex);
-						return IMediaSamples::EFetchBestSampleResult::Ok;
-					}
+				ImgMediaLoader::CheckAndUpdateImgDimensions(SequenceDim, Frame->Get()->Info.Dim);
+
+				if (Sample->Initialize(*Frame->Get(), SequenceDim, QueuedSampleFetch.LastTimeStamp, QueuedSampleFetch.LastDuration, GetNumMipLevels(), TilingDescription))
+				{
+					OutSample = Sample;
+					CSV_EVENT(ImgMedia, TEXT("LoaderFetchHit %d %d-%d"), MaxIdx, StartIndex, EndIndex);
+					return IMediaSamples::EFetchBestSampleResult::Ok;
 				}
 			}
 			else
@@ -769,8 +687,9 @@ bool FImgMediaLoader::PeekVideoSampleTime(FMediaTimeStamp &TimeStamp, bool bIsLo
 {
 	if (IsInitialized())
 	{
-		int32 Idx;
+		int32 Idx = INDEX_NONE;
 		bool bNewLoopSeq = false;
+		FTimespan FrameStart;
 
 		// Do we know which index we handed out last?
 		if (QueuedSampleFetch.LastFrameIndex != INDEX_NONE)
@@ -799,6 +718,8 @@ bool FImgMediaLoader::PeekVideoSampleTime(FMediaTimeStamp &TimeStamp, bool bIsLo
 					return false;
 				}
 			}
+
+			FrameStart = FrameNumberToTime(Idx);
 		}
 		else
 		{
@@ -816,7 +737,7 @@ bool FImgMediaLoader::PeekVideoSampleTime(FMediaTimeStamp &TimeStamp, bool bIsLo
 						NextTime -= SequenceDuration;
 						bNewLoopSeq = true;
 					}
-					Idx = TimeToFrameNumber(NextTime);
+					FrameStart = NextTime;
 				}
 				else
 				{
@@ -830,29 +751,26 @@ bool FImgMediaLoader::PeekVideoSampleTime(FMediaTimeStamp &TimeStamp, bool bIsLo
 						NextTime += SequenceDuration;
 						bNewLoopSeq = true;
 					}
-					Idx = TimeToFrameNumber(NextTime);
+					FrameStart = NextTime;
 				}
 			}
 			else
 			{
-				Idx = TimeToFrameNumber(CurrentTime);
+				FrameStart = CurrentTime;
 			}
 		}
 
-		auto SequenceIndex = QueuedSampleFetch.LastTimeStamp.IsValid() ? QueuedSampleFetch.LastTimeStamp.SequenceIndex : 0;
-
-		if (SkipFramesLevel > 0)
+		if (Idx == INDEX_NONE)
 		{
-			Idx = GetSkipFrame(Idx);
+			Idx = TimeToFrameNumber(FrameStart);
 		}
 
+		auto SequenceIndex = QueuedSampleFetch.LastTimeStamp.IsValid() ? QueuedSampleFetch.LastTimeStamp.SequenceIndex : 0;
+
 		// If possible, fetch any existing frame data...
 		// (just to see if we have any)
 		const TSharedPtr<FImgMediaFrame, ESPMode::ThreadSafe>* Frame = ((uint32)Idx != INDEX_NONE) ? (UseGlobalCache ? GlobalCache->FindAndTouch(SequenceName, Idx) : Frames.FindAndTouch(Idx)) : nullptr;
 
-		// Start time of this frame
-		FTimespan FrameStart = FrameNumberToTime(Idx);
-
 		// Data is present?
 		if (Frame)
 		{
@@ -864,7 +782,7 @@ bool FImgMediaLoader::PeekVideoSampleTime(FMediaTimeStamp &TimeStamp, bool bIsLo
 		else
 		{
 			// No data. We will request it (so, like other players do) we fill our (virtual) queue at the current location automatically
-			RequestFrame(FrameStart, PlayRate, bIsLoopingEnabled);
+			RequestFrame(Idx, PlayRate, bIsLoopingEnabled);
 		}
 	}
 	return false;
@@ -976,18 +894,22 @@ void FImgMediaLoader::Initialize(const FString& SequencePath, const FFrameRate&
 
 bool FImgMediaLoader::RequestFrame(FTimespan Time, float PlayRate, bool Loop)
 {
-	const int32 FrameNumber = TimeToFrameNumber(Time);
+	return RequestFrame(TimeToFrameNumber(Time), PlayRate, Loop);
+}
 
+
+bool FImgMediaLoader::RequestFrame(int32 FrameNumber, float PlayRate, bool Loop)
+{
 	if ((FrameNumber == INDEX_NONE) || (FrameNumber == LastRequestedFrame))
 	{
 		// Make sure we call the reader even if we do no update - just in case it does anything
 		Reader->OnTick();
 
-		UE_LOG(LogImgMedia, VeryVerbose, TEXT("Loader %p: Skipping frame %i for time %s last %d"), this, FrameNumber, *Time.ToString(TEXT("%h:%m:%s.%t")), LastRequestedFrame);
+		UE_LOG(LogImgMedia, VeryVerbose, TEXT("Loader %p: Skipping frame %i for time %s last %d"), this, FrameNumber, *FrameNumberToTime(FrameNumber).ToString(TEXT("%h:%m:%s.%t")), LastRequestedFrame);
 		return false;
 	}
 
-	UE_LOG(LogImgMedia, VeryVerbose, TEXT("Loader %p: Requesting frame %i for time %s"), this, FrameNumber, *Time.ToString(TEXT("%h:%m:%s.%t")));
+	UE_LOG(LogImgMedia, VeryVerbose, TEXT("Loader %p: Requesting frame %i for time %s"), this, FrameNumber, *FrameNumberToTime(FrameNumber).ToString(TEXT("%h:%m:%s.%t")));
 
 	Update(FrameNumber, PlayRate, Loop);
 	LastRequestedFrame = FrameNumber;
@@ -1016,15 +938,6 @@ void FImgMediaLoader::FrameNumbersToTimeRanges(const TArray<int32>& FrameNumbers
 }
 
 
-FTimespan FImgMediaLoader::FrameNumberToTime(uint32 FrameNumber) const
-{
-	return FTimespan(FMath::DivideAndRoundNearest(
-		FrameNumber * SequenceFrameRate.Denominator * ETimespan::TicksPerSecond,
-		(int64)SequenceFrameRate.Numerator
-	));
-}
-
-
 bool FImgMediaLoader::LoadSequence(const FString& SequencePath, const FFrameRate& FrameRateOverride, bool Loop, FImgMediaFrameInfo& FirstFrameInfo)
 {
 	SCOPE_CYCLE_COUNTER(STAT_ImgMedia_LoaderLoadSequence);
@@ -1374,6 +1287,7 @@ void FImgMediaLoader::FindMips(const FString& SequencePath)
 	}
 }
 
+
 uint32 FImgMediaLoader::TimeToFrameNumber(FTimespan Time) const
 {
 	if ((Time < FTimespan::Zero()) || (Time >= SequenceDuration) || GetNumImages() < 1)
@@ -1381,15 +1295,22 @@ uint32 FImgMediaLoader::TimeToFrameNumber(FTimespan Time) const
 		return INDEX_NONE;
 	}
 
-	const double FrameTimeErrorTollerance = 0.0001;
+	/*
+	* Older versions of this code attempted to help frame selection by using an Epsilon value
+	* during frame selection to compensate for numerical issues. We now use frame vs. time range
+	* requested coverage analysis to select the frames, which inherently will compensate for most / all
+	* cases covered by the older code.
+	*/
 
-	// note: we snap to the next best whole frame index if the compute result is with in FrameTimeErrorTollerance * FrameDuration to avoid
-	// incorrect frame selection if the value passed in is just ever so slightly off
-	double FrameDuration = (double)SequenceFrameRate.Numerator / SequenceFrameRate.Denominator;
-	double Frame = Time.GetTotalSeconds() * FrameDuration;
-	double Epsilon = FrameTimeErrorTollerance * FrameDuration;
+	// note: do NOT assume 100% reversability between TimeToFrameNumber() and FrameNumberToTime()! The FTimespan "only" has 100ns resolution and does some rounding when converting from seconds -> this migth shift things around!
+	return static_cast<uint32>(FMath::FloorToInt32(Time.GetTotalSeconds() * (static_cast<double>(SequenceFrameRate.Numerator) / SequenceFrameRate.Denominator)));
+}
 
-	return FMath::Min(uint32(Frame + Epsilon), uint32(GetNumImages() - 1));
+
+FTimespan FImgMediaLoader::FrameNumberToTime(uint32 FrameNumber) const
+{
+	// note: do NOT assume 100% reversability between TimeToFrameNumber() and FrameNumberToTime()! The FTimespan "only" has 100ns resolution and does some rounding when converting from seconds -> this migth shift things around!
+	return FTimespan::FromSeconds(FrameNumber * (static_cast<double>(SequenceFrameRate.Denominator) / SequenceFrameRate.Numerator));
 }
 
 
@@ -1727,31 +1648,6 @@ FTimespan FImgMediaLoader::ModuloTime(FTimespan Time)
 }
 
 
-float FImgMediaLoader::GetFrameOverlap(uint32 FrameIndex, FTimespan StartTime, FTimespan EndTime) const
-{
-	check(StartTime <= EndTime);
-	if (StartTime == EndTime)
-	{
-		return 0.0f;
-	}
-
-	float Overlap = 0.0f;
-
-	// Set up ranges.
-	FTimespan FrameStartTime = FrameNumberToTime(FrameIndex);
-	FTimespan FrameEndTime = FrameStartTime + FrameNumberToTime(1);
-
-	TRange<FTimespan> FrameRange(FrameStartTime, FrameEndTime);
-	TRange<FTimespan> TimeRange(StartTime, EndTime);
-	TRange<FTimespan> OverlapRange = TRange<FTimespan>::Intersection(FrameRange, TimeRange);
-
-	// Get overlap size.
-	FTimespan OverlapTimespan = OverlapRange.Size<FTimespan>();
-	Overlap = OverlapTimespan.GetTotalSeconds();
-
-	return Overlap;
-}
-
 bool FImgMediaLoader::GetNumberAtEndOfString(int32 &Number, const FString& String) const
 {
 	bool bFoundNumber = false;
diff --git a/Engine/Plugins/Media/ImgMedia/Source/ImgMedia/Private/Player/ImgMediaPlayer.cpp b/Engine/Plugins/Media/ImgMedia/Source/ImgMedia/Private/Player/ImgMediaPlayer.cpp
--- a/Engine/Plugins/Media/ImgMedia/Source/ImgMedia/Private/Player/ImgMediaPlayer.cpp
+++ b/Engine/Plugins/Media/ImgMedia/Source/ImgMedia/Private/Player/ImgMediaPlayer.cpp
@@ -654,7 +654,7 @@ bool FImgMediaPlayer::Seek(const FTimespan& Time)
 	CurrentSeekIndex += (CurrentRate >= 0.0f) ? 1 : -1;
 	if (Loader.IsValid())
 	{
-		Loader->Seek(FMediaTimeStamp(CurrentTime, FMediaTimeStamp::MakeSequenceIndex(CurrentSeekIndex, 0)));
+		Loader->Seek(FMediaTimeStamp(CurrentTime, FMediaTimeStamp::MakeSequenceIndex(CurrentSeekIndex, 0)), LastNonZeroRate < 0.0f);
 	}
 #endif // IMG_MEDIA_PLAYER_VERSION == 1
 
diff --git a/Engine/Plugins/Media/MediaCompositing/Source/MediaCompositing/Private/MovieScene/MovieSceneMediaTemplate.cpp b/Engine/Plugins/Media/MediaCompositing/Source/MediaCompositing/Private/MovieScene/MovieSceneMediaTemplate.cpp
--- a/Engine/Plugins/Media/MediaCompositing/Source/MediaCompositing/Private/MovieScene/MovieSceneMediaTemplate.cpp
+++ b/Engine/Plugins/Media/MediaCompositing/Source/MediaCompositing/Private/MovieScene/MovieSceneMediaTemplate.cpp
@@ -196,17 +196,7 @@ struct FMediaSectionExecutionToken
 		//
 
 		// Setup media time (used for seeks)
-		FTimespan MediaTime;
-		if (!MediaPlayer->IsLooping())
-		{
-			// note: we use a small offset at the end to make sure we can indeed seek to it (exclusive end type range)
-			MediaTime = FMath::Clamp(CurrentTime, FTimespan::Zero(), MediaDuration - FrameDuration * 0.5);
-		}
-		else
-		{
-			// one always seeks into the original media time-range, hence: modulo the time
-			MediaTime = CurrentTime % MediaDuration;
-		}
+		FTimespan MediaTime = CurrentTime;
 
 		#if MOVIESCENEMEDIATEMPLATE_TRACE_EVALUATION
 			GLog->Logf(ELogVerbosity::Log, TEXT("Executing time %s, MediaTime %s"), *CurrentTime.ToString(TEXT("%h:%m:%s.%t")), *MediaTime.ToString(TEXT("%h:%m:%s.%t")));
diff --git a/Engine/Source/Runtime/Media/Public/IMediaTimeSource.h b/Engine/Source/Runtime/Media/Public/IMediaTimeSource.h
--- a/Engine/Source/Runtime/Media/Public/IMediaTimeSource.h
+++ b/Engine/Source/Runtime/Media/Public/IMediaTimeSource.h
@@ -44,12 +44,12 @@ public:
 
 	static int64 MakeSequenceIndex(int32 PrimaryIndex, int32 SecondaryIndex)
 	{
-		return (int64(PrimaryIndex) << 32) + int64(SecondaryIndex);
+		return (static_cast<int64>(PrimaryIndex) << 32) + int64(SecondaryIndex);
 	}
 
 	static int64 AdjustPrimaryIndex(int64 InSequenceIndex, int32 Add)
 	{
-		return InSequenceIndex + (int64(Add) << 32);
+		return InSequenceIndex + (static_cast<int64>(Add) << 32);
 	}
 
 	static int64 AdjustSecondaryIndex(int64 InSequenceIndex, int32 Add)
@@ -60,12 +60,12 @@ public:
 	static int32 GetPrimaryIndex(int64 InSequenceIndex)
 	{
 		// note: needs to cope with signed operation
-		return int32((InSequenceIndex - int32(InSequenceIndex)) >> 32);	//-V1065 PSV warning "Expression can be simplified" suppressed
+		return static_cast<int32>((InSequenceIndex - int32(InSequenceIndex)) >> 32);	//-V1065 PSV warning "Expression can be simplified" suppressed
 	}
 
 	static int32 GetSecondaryIndex(int64 InSequenceIndex)
 	{
-		return int32(InSequenceIndex);
+		return static_cast<int32>(InSequenceIndex);
 	}
 
 	FMediaTimeStamp operator + (const FTimespan & Other) const { return FMediaTimeStamp(Time + Other, SequenceIndex); }
diff --git a/Engine/Source/Runtime/MediaUtils/Private/MediaPlayerFacade.cpp b/Engine/Source/Runtime/MediaUtils/Private/MediaPlayerFacade.cpp
--- a/Engine/Source/Runtime/MediaUtils/Private/MediaPlayerFacade.cpp
+++ b/Engine/Source/Runtime/MediaUtils/Private/MediaPlayerFacade.cpp
@@ -1151,11 +1151,32 @@ void FMediaPlayerFacade::QueryCacheState(EMediaTrackType TrackType, EMediaCacheS
 }
 
 
-bool FMediaPlayerFacade::Seek(const FTimespan& Time)
+bool FMediaPlayerFacade::Seek(const FTimespan& InTime)
 {
 	auto CurrentPlayer = Player;
 
-	if (!CurrentPlayer.IsValid() || !CurrentPlayer->GetControls().Seek(Time))
+	if (!CurrentPlayer.IsValid())
+	{
+		return false;
+	}
+
+	FTimespan Duration = CurrentPlayer->GetControls().GetDuration();
+	if (Duration == FTimespan::Zero())
+	{
+		return false;
+	}
+
+	FTimespan Time;
+	if (CurrentPlayer->GetControls().IsLooping())
+	{
+		Time = WrappedModulo(InTime, Duration);
+	}
+	else
+	{
+		Time = FTimespan(FMath::Clamp(InTime.GetTicks(), 0, Duration.GetTicks()));
+	}
+
+	if (!CurrentPlayer->GetControls().Seek(Time))
 	{
 		return false;
 	}
@@ -1234,9 +1255,8 @@ void FMediaPlayerFacade::SetBlockOnTimeRange(const TRange<FTimespan>& TimeRange)
 
 void FMediaPlayerFacade::FBlockOnRange::OnFlush()
 {
-	LastBlockOnRange = TRange<FTimespan>::Empty();
+	LastTimeRange = TRange<FTimespan>::Empty();
 	OnBlockPrimaryIndex = 0;
-	OnBlockSecondaryIndex = 0;
 }
 
 
@@ -1245,9 +1265,8 @@ void FMediaPlayerFacade::FBlockOnRange::OnSeek(int32 PrimaryIndex)
 	TSharedPtr<IMediaPlayer, ESPMode::ThreadSafe> CurrentPlayer(Facade->Player);
 	check(CurrentPlayer.IsValid());
 
-	LastBlockOnRange = TRange<FTimespan>::Empty();
+	LastTimeRange = TRange<FTimespan>::Empty();
 	OnBlockPrimaryIndex = PrimaryIndex;
-	OnBlockSecondaryIndex = 0;
 }
 
 
@@ -1287,7 +1306,7 @@ const TRange<FMediaTimeStamp>& FMediaPlayerFacade::FBlockOnRange::GetRange() con
 	// If the range is empty or the player can't support blocked playback: reset everything & return empty block range...
 	if (CurrentTimeRange.IsEmpty() || !CurrentPlayer->GetControls().CanControl(EMediaControl::BlockOnFetch))
 	{
-		LastBlockOnRange = TRange<FTimespan>::Empty();
+		LastTimeRange = TRange<FTimespan>::Empty();
 		BlockOnRange = TRange<FMediaTimeStamp>::Empty();
 		CurrentPlayer->GetControls().SetBlockingPlaybackHint(false);
 		return BlockOnRange;
@@ -1320,7 +1339,7 @@ const TRange<FMediaTimeStamp>& FMediaPlayerFacade::FBlockOnRange::GetRange() con
 
 	if (!CurrentPlayer->GetControls().IsLooping())
 	{
-		int64 SequenceIndex = FMediaTimeStamp::MakeSequenceIndex(OnBlockPrimaryIndex, OnBlockSecondaryIndex);
+		int64 SequenceIndex = FMediaTimeStamp::MakeSequenceIndex(OnBlockPrimaryIndex, 0);
 		BlockOnRange = TRange<FMediaTimeStamp>(FMediaTimeStamp(Start, SequenceIndex), FMediaTimeStamp(End, SequenceIndex));
 	}
 	else
@@ -1348,51 +1367,13 @@ const TRange<FMediaTimeStamp>& FMediaPlayerFacade::FBlockOnRange::GetRange() con
 		Start = WrappedModulo(Start, Duration);
 		End = WrappedModulo(End, Duration);
 
-		/*
-		* Note about moving along over multiple loops at once:
-		* We do not explicitly care about this case, but rather look at the blocking range to always be in the "first" loop, even if - on the external time line - we might be on a higher loop count.
-		* Hence we only need to detect if we moved "backwards" without an explicit seek vs. the last requested blocked range (see below) and then increase the loop index by one.
-		* 
-		* -> visually this will result in the desired behavior
-		*/
-
-		// Detect any non-monotonic movement of the range...
-		// (seeks will reset the last value, so this is ensured not to trigger in that case)
-		if (!LastBlockOnRange.IsEmpty())
-		{
-			if (Rate >= 0.0f)
-			{
-				FTimespan LastStart = WrappedModulo(LastBlockOnRange.GetLowerBoundValue(), Duration);
-				if ((LastStart > Start) || ((LastStart == Start) && (LastBlockOnRange.GetLowerBoundValue() < CurrentTimeRange.GetLowerBoundValue())))
-				{
-					++OnBlockSecondaryIndex;
-				}
-			}
-			else
-			{
-				FTimespan LastEnd = WrappedModulo(LastBlockOnRange.GetUpperBoundValue(), Duration);
-				if ((LastEnd < End) || ((LastEnd == End) && (LastBlockOnRange.GetUpperBoundValue() > CurrentTimeRange.GetUpperBoundValue())))
-				{
-					--OnBlockSecondaryIndex;
-				}
-			}
-		}
+		int64 LoopIdxS = FMath::FloorToInt(CurrentTimeRange.GetLowerBoundValue().GetTotalSeconds() / Duration.GetTotalSeconds());
+		int64 LoopIdxE = FMath::FloorToInt(CurrentTimeRange.GetUpperBoundValue().GetTotalSeconds() / Duration.GetTotalSeconds());
 
-		// Check if our range crosses the loop point (sequence boundary)
-		int32 StartLoopIndex, EndLoopIndex;
-		if (Rate >= 0.0)
-		{
-			StartLoopIndex = OnBlockSecondaryIndex;
-			EndLoopIndex = (Start <= End) ? OnBlockSecondaryIndex : (OnBlockSecondaryIndex + 1);
-		}
-		else
-		{
-			StartLoopIndex = (Start <= End) ? OnBlockSecondaryIndex : (OnBlockSecondaryIndex - 1);
-			EndLoopIndex = OnBlockSecondaryIndex;
-		}
+		int32 EndLoopIndex = LoopIdxE - LoopIdxS;
 
 		// Assemble final blocking range
-		auto SeqIndexStart = FMediaTimeStamp::MakeSequenceIndex(OnBlockPrimaryIndex, StartLoopIndex);
+		auto SeqIndexStart = FMediaTimeStamp::MakeSequenceIndex(OnBlockPrimaryIndex, 0);
 		auto SeqIndexEnd = FMediaTimeStamp::MakeSequenceIndex(OnBlockPrimaryIndex, EndLoopIndex);
 		BlockOnRange = TRange<FMediaTimeStamp>(FMediaTimeStamp(Start, SeqIndexStart), FMediaTimeStamp(End, SeqIndexEnd));
 		check(!BlockOnRange.IsEmpty());
@@ -1415,7 +1396,7 @@ const TRange<FMediaTimeStamp>& FMediaPlayerFacade::FBlockOnRange::GetRange() con
 
 	CurrentPlayer->GetControls().SetBlockingPlaybackHint(!BlockOnRange.IsEmpty());
 
-	LastBlockOnRange = CurrentTimeRange;
+	LastTimeRange = CurrentTimeRange;
 	RangeIsDirty = false;
 
 	return BlockOnRange;
@@ -1572,7 +1553,7 @@ bool FMediaPlayerFacade::BlockOnFetch() const
 {
 	check(Player.IsValid());
 
-	const TRange<FMediaTimeStamp>& BR = BlockOnRange.GetRange();
+	const TRange<FMediaTimeStamp> BR(GetAdjustedBlockOnRange());
 
 	if (BR.IsEmpty() || !Player->GetControls().CanControl(EMediaControl::BlockOnFetch) || BlockOnRangeDisabled || bHaveActiveAudio)
 	{
@@ -1644,7 +1625,7 @@ bool FMediaPlayerFacade::BlockOnFetch() const
 			return false; // block only in forward play
 		}
 
-		const bool VideoReady = (VideoSampleSinks.Num() == 0) || (BR.GetUpperBoundValue() < NextVideoSampleTime);
+		const bool VideoReady = (VideoSampleSinks.Num() == 0) || (BR.GetUpperBoundValue().Time < NextVideoSampleTime);
 
 		if (VideoReady)
 		{
@@ -2598,7 +2579,7 @@ bool FMediaPlayerFacade::GetCurrentPlaybackTimeRange(TRange<FMediaTimeStamp>& Ti
 			// Do we now have a current timestamp estimation?
 			if (!NextEstVideoTimeAtFrameStart.IsValid())
 			{
-				// No timing info available, no timerange available, no samples to process
+				// No timing info available, no time range available, no samples to process
 				return false;
 			}
 			else
@@ -2625,17 +2606,23 @@ bool FMediaPlayerFacade::GetCurrentPlaybackTimeRange(TRange<FMediaTimeStamp>& Ti
 			// External clock delivers time-range
 			// (for now we just use the blocking time range as this clock type is solely used in that case)
 			//
-			TimeRange = BlockOnRange.GetRange();
+			TimeRange = GetAdjustedBlockOnRange();
 		}
 	}
 
+	if (TimeRange.IsEmpty())
+	{
+		return false;
+	}
+
 	// If we are looping we check to prepare proper ranges should we wrap around either end of the media...
 	// (we do not clamp in the non-looping case as the rest of the code should deal with that fine)
 	if (Player->GetControls().IsLooping())
 	{
-		if (Player->GetControls().GetDuration() > 0.0f)
+		// Can the player already hand out a duration? (we might get here very early on)
+		const FTimespan Duration = Player->GetControls().GetDuration();
+		if (Duration != FTimespan::Zero())
 		{
-			const FTimespan Duration = Player->GetControls().GetDuration();
 			FTimespan WrappedStart = WrappedModulo(TimeRange.GetLowerBoundValue().Time, Duration);
 			FTimespan WrappedEnd = WrappedModulo(TimeRange.GetUpperBoundValue().Time, Duration);
 			if (WrappedStart > WrappedEnd)
@@ -2656,6 +2643,39 @@ bool FMediaPlayerFacade::GetCurrentPlaybackTimeRange(TRange<FMediaTimeStamp>& Ti
 }
 
 
+TRange<FMediaTimeStamp> FMediaPlayerFacade::GetAdjustedBlockOnRange() const
+{
+	TRange<FMediaTimeStamp> TimeRange = BlockOnRange.GetRange();
+
+	// The blocking range returned will base any secondary index values on zero. We need to patch in the secondary index value as last returned by the player plugin
+	if (!LastVideoSampleProcessedTimeRange.IsEmpty())
+	{
+		int32 BaseLoopIndex = FMediaTimeStamp::GetSecondaryIndex(LastVideoSampleProcessedTimeRange.GetLowerBoundValue().SequenceIndex);
+
+		// We still might have a time that is before the current one, in which case we need to push the loop index up (or down)...
+		if (GetUnpausedRate() >= 0.0f)
+		{
+			if (LastVideoSampleProcessedTimeRange.GetLowerBoundValue() > FMediaTimeStamp(TimeRange.GetLowerBoundValue().Time, FMediaTimeStamp::AdjustSecondaryIndex(TimeRange.GetLowerBoundValue().SequenceIndex, BaseLoopIndex)))
+			{
+				++BaseLoopIndex;
+			}
+		}
+		else
+		{
+			if (LastVideoSampleProcessedTimeRange.GetLowerBoundValue() < FMediaTimeStamp(TimeRange.GetLowerBoundValue().Time, FMediaTimeStamp::AdjustSecondaryIndex(TimeRange.GetLowerBoundValue().SequenceIndex, BaseLoopIndex)))
+			{
+				--BaseLoopIndex;
+			}
+		}
+
+		TimeRange.SetLowerBoundValue(FMediaTimeStamp(TimeRange.GetLowerBoundValue().Time, FMediaTimeStamp::AdjustSecondaryIndex(TimeRange.GetLowerBoundValue().SequenceIndex, BaseLoopIndex)));
+		TimeRange.SetUpperBoundValue(FMediaTimeStamp(TimeRange.GetUpperBoundValue().Time, FMediaTimeStamp::AdjustSecondaryIndex(TimeRange.GetUpperBoundValue().SequenceIndex, BaseLoopIndex)));
+	}
+
+	return TimeRange;
+}
+
+
 /* FMediaPlayerFacade implementation
 *****************************************************************************/
 
@@ -2878,26 +2898,20 @@ bool FMediaPlayerFacade::ProcessVideoSamples(IMediaSamples& Samples, const TRang
 
 		TRange<FMediaTimeStamp> SampleTimeRange(Sample->GetTime(), Sample->GetTime() + Sample->GetDuration());
 
-		// Is it what we want?
-		const TRange<FMediaTimeStamp>& BR = BlockOnRange.GetRange();
-		if (BR.IsEmpty() || BR.Overlaps(SampleTimeRange))
-		{
-			// Enqueue the sample to render
-			// (we use a queue to stay compatible with existing structure and older sinks - new sinks will read this single entry right away on the gamethread
-			//  and pass it along to rendering outside the queue)
-			bool bOk = VideoSampleSinks.Enqueue(Sample.ToSharedRef(), FMediaPlayerQueueDepths::MaxVideoSinkDepth);
-			check(bOk);
+		// Enqueue the sample to render
+		// (we use a queue to stay compatible with existing structure and older sinks - new sinks will read this single entry right away on the gamethread
+		//  and pass it along to rendering outside the queue)
+		bool bOk = VideoSampleSinks.Enqueue(Sample.ToSharedRef(), FMediaPlayerQueueDepths::MaxVideoSinkDepth);
+		check(bOk);
 
-			{
-				FScopeLock Lock(&LastTimeValuesCS);
-				CurrentFrameVideoDisplayTimeStamp = CurrentFrameVideoTimeStamp = SampleTimeRange.GetLowerBoundValue();
-				LastVideoSampleProcessedTimeRange = SampleTimeRange;
-			}
-
-			UpdateSeekStatus(&CurrentFrameVideoTimeStamp);
-
-			return true;
+		{
+			FScopeLock Lock(&LastTimeValuesCS);
+			CurrentFrameVideoDisplayTimeStamp = CurrentFrameVideoTimeStamp = SampleTimeRange.GetLowerBoundValue();
+			LastVideoSampleProcessedTimeRange = SampleTimeRange;
 		}
+
+		UpdateSeekStatus(&CurrentFrameVideoTimeStamp);
+		return true;
 	}
 	return false;
 }
diff --git a/Engine/Source/Runtime/MediaUtils/Public/MediaPlayerFacade.h b/Engine/Source/Runtime/MediaUtils/Public/MediaPlayerFacade.h
--- a/Engine/Source/Runtime/MediaUtils/Public/MediaPlayerFacade.h
+++ b/Engine/Source/Runtime/MediaUtils/Public/MediaPlayerFacade.h
@@ -861,6 +861,7 @@ private:
 	FMediaVideoSampleSinks VideoSampleSinks;
 
 private:
+	TRange<FMediaTimeStamp> GetAdjustedBlockOnRange() const;
 
 	class FBlockOnRange
 	{
@@ -879,10 +880,9 @@ private:
 		{
 			BlockOnRange = TRange<FMediaTimeStamp>::Empty();
 			CurrentTimeRange = TRange<FTimespan>::Empty();
-			LastBlockOnRange = TRange<FTimespan>::Empty();
+			LastTimeRange = TRange<FTimespan>::Empty();
 			RangeIsDirty = false;
 			OnBlockPrimaryIndex = 0;
-			OnBlockSecondaryIndex = 0;
 		}
 
 	private:
@@ -896,15 +896,13 @@ private:
 		mutable TRange<FMediaTimeStamp> BlockOnRange;
 
 		/** Last user provided BlockOnRange value */
-		mutable TRange<FTimespan> LastBlockOnRange;
+		mutable TRange<FTimespan> LastTimeRange;
 
 		/** Flag to indicate if internal range is valid or not */
 		mutable bool RangeIsDirty;
 
 		/** Primary sequence index used during blocked playback processing */
 		mutable int32 OnBlockPrimaryIndex;
-		/** Secondary sequence index used during blocked playback processing */
-		mutable int32 OnBlockSecondaryIndex;
 	};
 
 	FBlockOnRange BlockOnRange;
