ImgMedia: Fix exrgpureader's viewport derivation to include all pre-cached tiles.

diff --git a/Engine/Plugins/Media/ImgMedia/Source/ImgMedia/Internal/Assets/ImgMediaMipMapInfo.h b/Engine/Plugins/Media/ImgMedia/Source/ImgMedia/Internal/Assets/ImgMediaMipMapInfo.h
--- a/Engine/Plugins/Media/ImgMedia/Source/ImgMedia/Internal/Assets/ImgMediaMipMapInfo.h
+++ b/Engine/Plugins/Media/ImgMedia/Source/ImgMedia/Internal/Assets/ImgMediaMipMapInfo.h
@@ -108,6 +108,13 @@ struct FImgMediaTileSelection
 	 */
 	bool Contains(const FImgMediaTileSelection& Other) const;
 
+	/**
+	 * Combine the current selection with another (assumes they are the same size).
+	 *
+	 * @param Other Selection to include.
+	 */
+	void Include(const FImgMediaTileSelection& Other);
+
 	/**
 	 * Mark a tile as visible.
 	 *
diff --git a/Engine/Plugins/Media/ImgMedia/Source/ImgMedia/Private/Assets/ImgMediaMipMapInfo.cpp b/Engine/Plugins/Media/ImgMedia/Source/ImgMedia/Private/Assets/ImgMediaMipMapInfo.cpp
--- a/Engine/Plugins/Media/ImgMedia/Source/ImgMedia/Private/Assets/ImgMediaMipMapInfo.cpp
+++ b/Engine/Plugins/Media/ImgMedia/Source/ImgMedia/Private/Assets/ImgMediaMipMapInfo.cpp
@@ -101,6 +101,13 @@ bool FImgMediaTileSelection::Contains(const FImgMediaTileSelection& Other) const
 	return true;
 }
 
+void FImgMediaTileSelection::Include(const FImgMediaTileSelection& Other)
+{
+	ensure(Tiles.Num() == Other.Tiles.Num());
+
+	Tiles = TBitArray<>::BitwiseOR(Tiles, Other.Tiles, EBitwiseOperatorFlags::MaxSize);
+}
+
 void FImgMediaTileSelection::SetVisible(int32 TileCoordX, int32 TileCoordY)
 {
 	int Index = ToIndex(FMath::Min(TileCoordX, Dimensions.X - 1), FMath::Min(TileCoordY, Dimensions.Y - 1), Dimensions);
diff --git a/Engine/Plugins/Media/ImgMedia/Source/ImgMedia/Private/Readers/ExrImgMediaReaderGpu.cpp b/Engine/Plugins/Media/ImgMedia/Source/ImgMedia/Private/Readers/ExrImgMediaReaderGpu.cpp
--- a/Engine/Plugins/Media/ImgMedia/Source/ImgMedia/Private/Readers/ExrImgMediaReaderGpu.cpp
+++ b/Engine/Plugins/Media/ImgMedia/Source/ImgMedia/Private/Readers/ExrImgMediaReaderGpu.cpp
@@ -175,32 +175,20 @@ FExrImgMediaReader::EReadResult FExrImgMediaReaderGpu::ReadMip
 		if (bHasTiles || ConverterParams->bCustomExr)
 		{
 			TArray<FIntRect> TileRegionsToRead;
-			TArray<FIntRect> TileRegionsToRender;
 			{
 				TRACE_CPUPROFILER_EVENT_SCOPE_TEXT(*FString::Printf(TEXT("ExrReaderGpu.CalculateRegions %d"), CurrentMipLevel));
 
 				if (const FImgMediaTileSelection* CachedSelection = OutFrame->MipTilesPresent.Find(CurrentMipLevel))
 				{
 					TileRegionsToRead = CachedSelection->GetVisibleRegions(&CurrentTileSelection);
-					TileRegionsToRender = CurrentTileSelection.GetVisibleRegions();
 				}
 				else
 				{
-					TileRegionsToRead = TileRegionsToRender = CurrentTileSelection.GetVisibleRegions();
+					TileRegionsToRead = CurrentTileSelection.GetVisibleRegions();
 				}
 			}
 
-			TArray<FIntRect>& Viewports = ConverterParams->Viewports.Add(CurrentMipLevel);
-			for (const FIntRect& TileRegion : TileRegionsToRender)
-			{
-				FIntRect Viewport;
-				Viewport.Min = FIntPoint(ConverterParams->TileDimWithBorders.X * TileRegion.Min.X, ConverterParams->TileDimWithBorders.Y * TileRegion.Min.Y);
-				Viewport.Max = FIntPoint(ConverterParams->TileDimWithBorders.X * TileRegion.Max.X, ConverterParams->TileDimWithBorders.Y * TileRegion.Max.Y);
-				Viewport.Clip(FIntRect(FIntPoint::ZeroValue, CurrentMipDim));
-				Viewports.Add(MoveTemp(Viewport));
-			}
-
-			if (TileRegionsToRead.IsEmpty() && !TileRegionsToRender.IsEmpty())
+			if (TileRegionsToRead.IsEmpty() && CurrentTileSelection.IsAnyVisible())
 			{
 				// If all tiles were previously read and stored in cached frame, reading can be skipped.
 				ReadResult = Skipped;
@@ -217,13 +205,6 @@ FExrImgMediaReader::EReadResult FExrImgMediaReaderGpu::ReadMip
 		}
 		else
 		{
-			TArray<FIntRect>& Viewports = ConverterParams->Viewports.Add(CurrentMipLevel);
-			FIntRect Viewport;
-
-			Viewport.Min = FIntPoint(0, 0);
-			Viewport.Max = CurrentMipDim;
-			Viewports.Add(MoveTemp(Viewport));
-
 			ReadResult = ReadInChunks(MipDataPtr, ImagePath, ConverterParams->FrameId, CurrentMipDim, BufferSize);
 			OutFrame->NumTilesRead++;
 		}
@@ -347,36 +328,76 @@ bool FExrImgMediaReaderGpu::ReadFrame(int32 FrameId, const TMap<int32, FImgMedia
 			FString ImagePath = Loader->GetImagePath(ConverterParams->FrameId, ConverterParams->bMipsInSeparateFiles ? CurrentMipLevel : 0);
 
 			EReadResult ReadResult = ReadMip(CurrentMipLevel, CurrentTileSelection, OutFrame, ConverterParams, SampleConverter, ImagePath, bHasTiles);
-
-			/* Error handling. */
-			if (ReadResult == Success)
-			{
-				OutFrame->MipTilesPresent.Emplace(CurrentMipLevel, CurrentTileSelection);
-			}
-			else if (ReadResult == Fail)
+			switch (ReadResult)
 			{
-				// Check if we have a compressed file.
-				FImgMediaFrameInfo Info;
-				if (GetInfo(ImagePath, Info))
+			case FExrImgMediaReader::Success:
+				if (OutFrame->MipTilesPresent.Find(CurrentMipLevel))
+				{
+					OutFrame->MipTilesPresent[CurrentMipLevel].Include(CurrentTileSelection);
+				}
+				else
+				{
+					OutFrame->MipTilesPresent.Emplace(CurrentMipLevel, CurrentTileSelection);
+				}
+				break;
+			case FExrImgMediaReader::Fail:
 				{
-					if (Info.CompressionName != "Uncompressed")
+					// Check if we have a compressed file.
+					FImgMediaFrameInfo Info;
+					if (GetInfo(ImagePath, Info))
 					{
-						UE_LOG(LogImgMedia, Error, TEXT("GPU Reader cannot read compressed file %s."), *ImagePath);
-						UE_LOG(LogImgMedia, Error, TEXT("Compressed and uncompressed files should not be mixed in a single sequence."));
+						if (Info.CompressionName != "Uncompressed")
+						{
+							UE_LOG(LogImgMedia, Error, TEXT("GPU Reader cannot read compressed file %s."), *ImagePath);
+							UE_LOG(LogImgMedia, Error, TEXT("Compressed and uncompressed files should not be mixed in a single sequence."));
+						}
 					}
-				}
 
-				// Fall back to CPU.
-				bFallBackToCPU = true;
+					// Fall back to CPU.
+					bFallBackToCPU = true;
 
-				// To make sure that Media Texture doesn't call the converter if this frame is invalid.
-				OutFrame->SampleConverter.Reset();
+					// To make sure that Media Texture doesn't call the converter if this frame is invalid.
+					OutFrame->SampleConverter.Reset();
 
-				return FExrImgMediaReader::ReadFrame(ConverterParams->FrameId, InMipTiles, OutFrame);
-			}
-			else if (ReadResult == Cancelled)
-			{
+					return FExrImgMediaReader::ReadFrame(ConverterParams->FrameId, InMipTiles, OutFrame);
+				}
+				break;
+			case FExrImgMediaReader::Cancelled:
+				// Abort further reading
 				return false;
+			case FExrImgMediaReader::Skipped:
+				// No new tiles were read, continue to the next mip level.
+				break;
+
+			default:
+				checkNoEntry();
+				break;
+			};
+		}
+		
+		// Create viewport(s) with all mip/tiles present
+		for (const TPair<int32, FImgMediaTileSelection>& TilesPerMip : OutFrame->MipTilesPresent)
+		{
+			const FImgMediaTileSelection& CurrentTileSelection = TilesPerMip.Value;
+			const int32 CurrentMipLevel = TilesPerMip.Key;
+			const int32 MipLevelDiv = 1 << CurrentMipLevel;
+			
+			TArray<FIntRect>& Viewports = ConverterParams->Viewports.Add(CurrentMipLevel);
+			for (const FIntRect& TileRegion : CurrentTileSelection.GetVisibleRegions())
+			{
+				FIntRect Viewport;
+				if (bHasTiles || ConverterParams->bCustomExr)
+				{
+					Viewport.Min = FIntPoint(ConverterParams->TileDimWithBorders.X * TileRegion.Min.X, ConverterParams->TileDimWithBorders.Y * TileRegion.Min.Y);
+					Viewport.Max = FIntPoint(ConverterParams->TileDimWithBorders.X * TileRegion.Max.X, ConverterParams->TileDimWithBorders.Y * TileRegion.Max.Y);
+					Viewport.Clip(FIntRect(FIntPoint::ZeroValue, CurrentMipDim));
+				}
+				else
+				{
+					Viewport.Min = FIntPoint(0, 0);
+					Viewport.Max = CurrentMipDim;
+				}
+				Viewports.Add(Viewport);
 			}
 		}
 	}
