Media Plate Exr: Fixed a crash caused by a camera movement. This was caused by media texture trying to render mips that were not read.

diff --git a/Engine/Plugins/Media/ImgMedia/Source/ImgMedia/Private/Readers/ExrImgMediaReaderGpu.cpp b/Engine/Plugins/Media/ImgMedia/Source/ImgMedia/Private/Readers/ExrImgMediaReaderGpu.cpp
--- a/Engine/Plugins/Media/ImgMedia/Source/ImgMedia/Private/Readers/ExrImgMediaReaderGpu.cpp
+++ b/Engine/Plugins/Media/ImgMedia/Source/ImgMedia/Private/Readers/ExrImgMediaReaderGpu.cpp
@@ -32,6 +32,7 @@
 DECLARE_GPU_STAT_NAMED(ExrImgMediaReaderGpu, TEXT("ExrImgGpu"));
 DECLARE_GPU_STAT_NAMED(ExrImgMediaReaderGpu_MipRender, TEXT("ExrImgGpu.MipRender"));
 DECLARE_GPU_STAT_NAMED(ExrImgMediaReaderGpu_MipUpscale, TEXT("ExrImgGpu.MipRender.MipUpscale"));
+DECLARE_GPU_STAT_NAMED(ExrImgMediaReaderGpu_CopyUploadBuffer, TEXT("ExrImgGpu.MipRender.UploadBufferCopy"));
 
 static TAutoConsoleVariable<bool> CVarExrReaderUseUploadHeap(
 	TEXT("r.ExrReaderGPU.UseUploadHeap"),
@@ -215,6 +216,8 @@ FExrImgMediaReader::EReadResult FExrImgMediaReaderGpu::ReadMip
 				{
 					SCOPED_DRAW_EVENT(RHICmdList, FExrImgMediaReaderGpu_CopyBuffers);
 					TRACE_CPUPROFILER_EVENT_SCOPE_TEXT(*FString::Printf(TEXT("ExrReaderGpu.StartCopy %d"), ConverterParams->FrameId));
+					SCOPED_GPU_STAT(RHICmdList, ExrImgMediaReaderGpu_CopyUploadBuffer);
+
 					if (BufferRegionsToCopy.IsEmpty())
 					{
 						RHICmdList.CopyBufferRegion(BufferData->ShaderAccessBufferRef, 0, BufferData->UploadBufferRef, 0, BufferData->ShaderAccessBufferRef->GetSize());
@@ -380,6 +383,13 @@ bool FExrImgMediaReaderGpu::ReadFrame(int32 FrameId, const TMap<int32, FImgMedia
 		{
 			const FImgMediaTileSelection& CurrentTileSelection = TilesPerMip.Value;
 			const int32 CurrentMipLevel = TilesPerMip.Key;
+
+			// Skip this viewport since we don't have anything to render.
+			if (!SampleConverter->GetMipLevelBuffer(CurrentMipLevel))
+			{
+				continue;
+			}
+
 			const int32 MipLevelDiv = 1 << CurrentMipLevel;
 			
 			TArray<FIntRect>& Viewports = ConverterParams->Viewports.Add(CurrentMipLevel);
@@ -622,6 +632,12 @@ void FExrImgMediaReaderGpu::CreateSampleConverterCallback
 			{
 				SCOPED_GPU_STAT(RHICmdList, ExrImgMediaReaderGpu_MipUpscale);
 
+				// Sanity check.
+				if (!MipBuffers.Contains(MipToUpscale))
+				{
+					UE_LOG(LogImgMedia, Warning, TEXT("Requested mip could not be found %s"), MipToUpscale);
+				}
+
 				FStructuredBufferPoolItemSharedPtr BufferDataToUpscale = MipBuffers[MipToUpscale];
 				TArray<FIntRect> FakeViewport;
 				FakeViewport.Add(FIntRect(FIntPoint(0, 0), Dim));
@@ -632,6 +648,13 @@ void FExrImgMediaReaderGpu::CreateSampleConverterCallback
 		for (const TPair<int32, TArray<FIntRect>>& MipLevelViewports : ConverterParams->Viewports)
 		{
 			int32 MipLevel = MipLevelViewports.Key;
+			
+			// Sanity check.
+			if (!MipBuffers.Contains(MipLevel))
+			{
+				continue;
+			}
+
 			FStructuredBufferPoolItemSharedPtr BufferData = MipBuffers[MipLevel];
 			int MipLevelDiv = 1 << MipLevel;
 			FIntPoint Dim = ConverterParams->FullResolution / MipLevelDiv;
