ImgMedia: Empty frames should not fill MipTilesPresent.

While the error is now avoided, this could have caused the discrepancy between viewports and buffers in UE-181575.

#rnx

diff --git a/Engine/Plugins/Media/ImgMedia/Source/ImgMedia/Internal/Loader/ImgMediaLoader.h b/Engine/Plugins/Media/ImgMedia/Source/ImgMedia/Internal/Loader/ImgMediaLoader.h
--- a/Engine/Plugins/Media/ImgMedia/Source/ImgMedia/Internal/Loader/ImgMediaLoader.h
+++ b/Engine/Plugins/Media/ImgMedia/Source/ImgMedia/Internal/Loader/ImgMediaLoader.h
@@ -475,11 +475,12 @@ protected:
 	void Update(int32 PlayHeadFrame, float PlayRate, bool Loop);
 
 	/**
-	 * Adds an empty frame to the cache.
+	 * Adds an empty frame to the cache if no frame is present at the specified number.
 	 * 
 	 * @param FrameNumber Frame number to add.
+	 * @return True if the frame was added
 	 */
-	void AddEmptyFrame(int32 FrameNumber);
+	bool TryAddEmptyFrame(int32 FrameNumber);
 
 	/**
 	 * Adds a frame to the cache.
diff --git a/Engine/Plugins/Media/ImgMedia/Source/ImgMedia/Private/Loader/ImgMediaLoader.cpp b/Engine/Plugins/Media/ImgMedia/Source/ImgMedia/Private/Loader/ImgMediaLoader.cpp
--- a/Engine/Plugins/Media/ImgMedia/Source/ImgMedia/Private/Loader/ImgMediaLoader.cpp
+++ b/Engine/Plugins/Media/ImgMedia/Source/ImgMedia/Private/Loader/ImgMediaLoader.cpp
@@ -56,21 +56,15 @@ namespace ImgMediaLoader
 	}
 	
 	// Check if the existing tiles contain all of the requested ones. (Is existing a superset of requested?)
-	bool ContainsMipTiles(const TSharedPtr<FImgMediaFrame, ESPMode::ThreadSafe>& ExistingFrame, const TMap<int32, FImgMediaTileSelection>& RequestedTiles)
+	bool ContainsMipTiles(const TMap<int32, FImgMediaTileSelection>& ExistingTiles, const TMap<int32, FImgMediaTileSelection>& RequestedTiles)
 	{
 		for (auto Iter = RequestedTiles.CreateConstIterator(); Iter; ++Iter)
 		{
 			const int32 RequestedMipLevel = Iter.Key();
 			const FImgMediaTileSelection& RequestedSelection = Iter.Value();
 
-			if (const FImgMediaTileSelection* ExistingSelection = ExistingFrame->MipTilesPresent.Find(RequestedMipLevel))
+			if (const FImgMediaTileSelection* ExistingSelection = ExistingTiles.Find(RequestedMipLevel))
 			{
-				if (ExistingSelection->GetDimensions().GetMax() < RequestedSelection.GetDimensions().GetMax())
-				{
-					// If an empty 1x1 frame was cached but we're now requesting a fully tiled frame.
-					return false;
-				}
-
 				if (!ExistingSelection->Contains(RequestedSelection))
 				{
 					// Requested tile selection is not present.
@@ -825,12 +819,7 @@ IQueuedWork* FImgMediaLoader::GetWork()
 	if (DesiredMipsAndTiles.IsEmpty())
 	{
 		// Still provide the cache with an empty frame to prevent blocking playback from stalling.
-		bool bIsFrameAlreadyCached = UseGlobalCache ? GlobalCache->Contains(SequenceName, FrameNumber) : Frames.Contains(FrameNumber);
-
-		if (!bIsFrameAlreadyCached)
-		{
-			AddEmptyFrame(FrameNumber);
-		}
+		TryAddEmptyFrame(FrameNumber);
 
 		// No selection was visible, so we don't queue any work.
 		return nullptr;
@@ -1482,7 +1471,7 @@ void FImgMediaLoader::Update(int32 PlayHeadFrame, float PlayRate, bool Loop)
 			TMap<int32, FImgMediaTileSelection> DesiredMipsAndTiles;
 			GetDesiredMipTiles(FrameNumber, DesiredMipsAndTiles);
 
-			NeedFrame = !ImgMediaLoader::ContainsMipTiles(*FramePtr, DesiredMipsAndTiles);
+			NeedFrame = !ImgMediaLoader::ContainsMipTiles((*FramePtr)->MipTilesPresent, DesiredMipsAndTiles);
 
 			if (NeedFrame)
 			{
@@ -1495,7 +1484,7 @@ void FImgMediaLoader::Update(int32 PlayHeadFrame, float PlayRate, bool Loop)
 			// Do we actually have a frame for this?
 			if (ImagePaths[0][FrameNumber].Len() == 0)
 			{
-				AddEmptyFrame(FrameNumber);
+				TryAddEmptyFrame(FrameNumber);
 			}
 			else
 			{
@@ -1530,25 +1519,28 @@ void FImgMediaLoader::Update(int32 PlayHeadFrame, float PlayRate, bool Loop)
  *****************************************************************************/
 
 
-void FImgMediaLoader::AddEmptyFrame(int32 FrameNumber)
+bool FImgMediaLoader::TryAddEmptyFrame(int32 FrameNumber)
 {
 	FScopeLock Lock(&CriticalSection);
 
-	TSharedPtr<FImgMediaFrame, ESPMode::ThreadSafe> Frame;
-	int32 NumChannels = 3;
-	int32 PixelSize = sizeof(uint16) * NumChannels;
-	Frame = MakeShareable(new FImgMediaFrame());
-	Frame->Info.Dim = SequenceDim;
-	Frame->Info.FrameRate = SequenceFrameRate;
-	Frame->Info.NumChannels = NumChannels;
-	Frame->Info.bHasTiles = false;
-	Frame->Format = EMediaTextureSampleFormat::FloatRGB;
-	Frame->Stride = Frame->Info.Dim.X * PixelSize;
-	for (int32 Level = 0; Level < GetNumMipLevels(); ++Level)
-	{
-		Frame->MipTilesPresent.Emplace(Level, FImgMediaTileSelection(1, 1, true));
-	}
-	AddFrameToCache(FrameNumber, Frame);
+	bool bAddFrame = UseGlobalCache ? !GlobalCache->Contains(SequenceName, FrameNumber) : !Frames.Contains(FrameNumber);
+
+	if (bAddFrame)
+	{
+		const int32 NumChannels = 3;
+		const int32 PixelSize = sizeof(uint16) * NumChannels;
+		TSharedPtr<FImgMediaFrame, ESPMode::ThreadSafe> Frame = MakeShared<FImgMediaFrame>();
+		Frame->Info.Dim = SequenceDim;
+		Frame->Info.FrameRate = SequenceFrameRate;
+		Frame->Info.NumChannels = NumChannels;
+		Frame->Info.bHasTiles = false;
+		Frame->Format = EMediaTextureSampleFormat::FloatRGB;
+		Frame->Stride = Frame->Info.Dim.X * PixelSize;
+
+		AddFrameToCache(FrameNumber, Frame);
+	}
+
+	return bAddFrame;
 }
 
 void FImgMediaLoader::NotifyWorkComplete(FImgMediaLoaderWork& CompletedWork, int32 FrameNumber,
