ImgMedia: Added a cvar that allows nodes to ignore all viewports that are not inner. Users will have to enable upscaling to get lower quality mips displayed in the outer.

diff --git a/Engine/Plugins/Media/ImgMedia/Source/ImgMedia/ImgMedia.Build.cs b/Engine/Plugins/Media/ImgMedia/Source/ImgMedia/ImgMedia.Build.cs
--- a/Engine/Plugins/Media/ImgMedia/Source/ImgMedia/ImgMedia.Build.cs
+++ b/Engine/Plugins/Media/ImgMedia/Source/ImgMedia/ImgMedia.Build.cs
@@ -24,6 +24,8 @@ namespace UnrealBuildTool.Rules
 					"RenderCore",
 					"Renderer",
 					"RHI",
+					"Slate",
+					"SlateCore",
 				});
 
 			PrivateIncludePathModuleNames.AddRange(
diff --git a/Engine/Plugins/Media/ImgMedia/Source/ImgMedia/Private/ImgMediaSceneViewExtension.cpp b/Engine/Plugins/Media/ImgMedia/Source/ImgMedia/Private/ImgMediaSceneViewExtension.cpp
--- a/Engine/Plugins/Media/ImgMedia/Source/ImgMedia/Private/ImgMediaSceneViewExtension.cpp
+++ b/Engine/Plugins/Media/ImgMedia/Source/ImgMedia/Private/ImgMediaSceneViewExtension.cpp
@@ -2,15 +2,38 @@
 
 #include "ImgMediaSceneViewExtension.h"
 #include "DynamicResolutionState.h"
-
+#include "Framework/Notifications/NotificationManager.h"
+#include "Widgets/Notifications/SNotificationList.h"
 #include "SceneView.h"
 
+#define LOCTEXT_NAMESPACE "ImgMediaSceneViewExtension"
+
 static TAutoConsoleVariable<float> CVarImgMediaFieldOfViewMultiplier(
 	TEXT("ImgMedia.FieldOfViewMultiplier"),
 	1.0f,
 	TEXT("Multiply the field of view for active cameras by this value, generally to increase the frustum overall sizes to mitigate missing tile artifacts.\n"),
 	ECVF_Default);
 
+static TAutoConsoleVariable<bool> CVarImgMediaProcessTilesInnerOnly(
+	TEXT("ImgMedia.ICVFX.InnerOnlyTiles"),
+	false,
+	TEXT("This CVar will ignore tile calculation for all viewports except for Display Cluster inner viewports. User should enable upscaling on Media plate to display lower quality mips instead, otherwise ")
+	TEXT("other viewports will only display tiles loaded specifically for inner viewport and nothing else. \n"),
+#if WITH_EDITOR
+	FConsoleVariableDelegate::CreateLambda([](IConsoleVariable* CVar)
+	{
+		if (CVar->GetBool())
+		{
+			FNotificationInfo Info(LOCTEXT("EnableUpscalingNotification", "Tile calculation enabled for Display Cluster Inner Viewports exclusively.\nUse Mip Upscaling option on Media Plate to fill empty texture areas with lower quality data."));
+			// Expire in 5 seconds.
+			Info.ExpireDuration = 5.0f;
+			FSlateNotificationManager::Get().AddNotification(Info);
+		}
+	}),
+#endif
+	ECVF_Default);
+
+
 FImgMediaSceneViewExtension::FImgMediaSceneViewExtension(const FAutoRegister& AutoReg)
 	: FSceneViewExtensionBase(AutoReg)
 	, CachedViewInfos()
@@ -68,6 +91,12 @@ int32 FImgMediaSceneViewExtension::GetPriority() const
 
 void FImgMediaSceneViewExtension::CacheViewInfo(FSceneViewFamily& InViewFamily, const FSceneView& View)
 {
+	// This relies on DisplayClusterMediaHelpers::GenerateICVFXViewportName to have two strings embedded.
+	if (CVarImgMediaProcessTilesInnerOnly.GetValueOnGameThread()
+		&& !(InViewFamily.ProfileDescription.Contains("_icvfx_") && InViewFamily.ProfileDescription.Contains("_incamera")))
+	{
+		return;
+	}
 	static const auto CVarMinAutomaticViewMipBiasOffset = IConsoleManager::Get().FindTConsoleVariableDataFloat(TEXT("r.ViewTextureMipBias.Offset"));
 	static const auto CVarMinAutomaticViewMipBias = IConsoleManager::Get().FindTConsoleVariableDataFloat(TEXT("r.ViewTextureMipBias.Min"));
 	const float FieldOfViewMultiplier = CVarImgMediaFieldOfViewMultiplier.GetValueOnGameThread();
@@ -146,3 +175,4 @@ void FImgMediaSceneViewExtension::CacheViewInfo(FSceneViewFamily& InViewFamily,
 
 	CachedViewInfos.Add(MoveTemp(Info));
 }
+#undef LOCTEXT_NAMESPACE
