ImgMedia: Deprecating UImgMediaPlaybackComponent.

#rnx

diff --git a/Engine/Plugins/Media/ImgMedia/Source/ImgMediaEngine/Private/Unreal/ImgMediaPlaybackComponent.cpp b/Engine/Plugins/Media/ImgMedia/Source/ImgMediaEngine/Private/Unreal/ImgMediaPlaybackComponent.cpp
--- a/Engine/Plugins/Media/ImgMedia/Source/ImgMediaEngine/Private/Unreal/ImgMediaPlaybackComponent.cpp
+++ b/Engine/Plugins/Media/ImgMedia/Source/ImgMediaEngine/Private/Unreal/ImgMediaPlaybackComponent.cpp
@@ -9,27 +9,27 @@
 
 #include UE_INLINE_GENERATED_CPP_BY_NAME(ImgMediaPlaybackComponent)
 
-UImgMediaPlaybackComponent::UImgMediaPlaybackComponent(const FObjectInitializer& ObjectInitializer)
+UDEPRECATED_ImgMediaPlaybackComponent::UDEPRECATED_ImgMediaPlaybackComponent(const FObjectInitializer& ObjectInitializer)
 	: Super(ObjectInitializer)
 {
 	bAutoActivate = true;
 }
 
-void UImgMediaPlaybackComponent::BeginPlay()
+void UDEPRECATED_ImgMediaPlaybackComponent::BeginPlay()
 {
 	Super::BeginPlay();
 
 	RegisterWithMipMapInfo();
 }
 
-void UImgMediaPlaybackComponent::EndPlay(const EEndPlayReason::Type EndPlayReason)
+void UDEPRECATED_ImgMediaPlaybackComponent::EndPlay(const EEndPlayReason::Type EndPlayReason)
 {
 	UnregisterWithMipMapInfo();
 
 	Super::EndPlay(EndPlayReason);
 }
 
-void UImgMediaPlaybackComponent::RegisterWithMipMapInfo()
+void UDEPRECATED_ImgMediaPlaybackComponent::RegisterWithMipMapInfo()
 {
 	// Get all texures used by our actor.
 	AActor* Owner = GetOwner();
@@ -73,7 +73,7 @@ void UImgMediaPlaybackComponent::RegisterWithMipMapInfo()
 	}
 }
 
-void UImgMediaPlaybackComponent::UnregisterWithMipMapInfo()
+void UDEPRECATED_ImgMediaPlaybackComponent::UnregisterWithMipMapInfo()
 {
 	if (ObjectInfo != nullptr)
 	{
diff --git a/Engine/Plugins/Media/ImgMedia/Source/ImgMediaEngine/Public/Unreal/ImgMediaPlaybackComponent.h b/Engine/Plugins/Media/ImgMedia/Source/ImgMediaEngine/Public/Unreal/ImgMediaPlaybackComponent.h
--- a/Engine/Plugins/Media/ImgMedia/Source/ImgMediaEngine/Public/Unreal/ImgMediaPlaybackComponent.h
+++ b/Engine/Plugins/Media/ImgMedia/Source/ImgMediaEngine/Public/Unreal/ImgMediaPlaybackComponent.h
@@ -14,12 +14,12 @@ struct FMediaTextureTrackerObject;
  * Not necessary to do this, but if you do then you can get additional functionality
  * such as selective mipmap loading.
  */
-UCLASS(ClassGroup = Media, meta = (BlueprintSpawnableComponent))
-class IMGMEDIAENGINE_API UImgMediaPlaybackComponent : public UActorComponent
+UCLASS(Deprecated)
+class UE_DEPRECATED(5.3, "ImgMediaPlaybackComponent functionality is deprecated. Please update your project to use the features of the MediaPlate instead.")  IMGMEDIAENGINE_API UDEPRECATED_ImgMediaPlaybackComponent : public UActorComponent
 {
 	GENERATED_BODY()
 public:
-	UImgMediaPlaybackComponent(const FObjectInitializer& ObjectInitializer);
+	UDEPRECATED_ImgMediaPlaybackComponent(const FObjectInitializer& ObjectInitializer);
 
 	/** This will be added to the calculated mipmap level. E.g. if set to 2, and you would normally be at mipmap level 1, then you will actually be at level 3. */
 	UPROPERTY(EditAnywhere, Category = MipMaps)
