ImgMedia ExrReaderGpu: A much better way of keeping track of pools. Since we've changed the way mip buffers are stored (in Frame cache) we are able to do better allocation and release of buffers that doesn't require staging pool, fences or boolean values.
General refactor and other thread safety improvements.

diff --git a/Engine/Plugins/Media/ImgMedia/Source/ImgMedia/Private/Readers/ExrImgMediaReader.cpp b/Engine/Plugins/Media/ImgMedia/Source/ImgMedia/Private/Readers/ExrImgMediaReader.cpp
--- a/Engine/Plugins/Media/ImgMedia/Source/ImgMedia/Private/Readers/ExrImgMediaReader.cpp
+++ b/Engine/Plugins/Media/ImgMedia/Source/ImgMedia/Private/Readers/ExrImgMediaReader.cpp
@@ -59,7 +59,6 @@ FExrImgMediaReader::EReadResult FExrImgMediaReader::ReadTiles
 	( uint16* Buffer
 	, int64 BufferSize
 	, const FString& ImagePath
-	, int32 FrameId
 	, const TArray<FIntRect>& TileRegions
 	, TSharedPtr<FSampleConverterParameters> ConverterParams
 	, const int32 CurrentMipLevel
@@ -113,9 +112,9 @@ FExrImgMediaReader::EReadResult FExrImgMediaReader::ReadTiles
 				// Check to see if the frame was canceled.
 				{
 					FScopeLock RegionScopeLock(&CanceledFramesCriticalSection);
-					if (CanceledFrames.Remove(FrameId) > 0)
+					if (CanceledFrames.Remove(ConverterParams->FrameId) > 0)
 					{
-						UE_LOG(LogImgMedia, Verbose, TEXT("Reader %p: Canceling Frame %i At tile row # %i"), this, FrameId, TileRow);
+						UE_LOG(LogImgMedia, Verbose, TEXT("Reader %p: Canceling Frame %i At tile row # %i"), this, ConverterParams->FrameId, TileRow);
 						Result = Cancelled;
 						break;
 					}
@@ -272,9 +271,10 @@ bool FExrImgMediaReader::ReadFrame(int32 FrameId, const TMap<int32, FImgMediaTil
 					ConverterParams->TileDimWithBorders = OutFrame->Info.TileDimensions + OutFrame->Info.TileBorder * 2;
 					ConverterParams->NumMipLevels = Loader->GetNumMipLevels();
 					ConverterParams->bCustomExr = OutFrame->Info.FormatName == TEXT("EXR CUSTOM");
+					ConverterParams->FrameId = FrameId;
 
 					TArray<UE::Math::TIntPoint<int64>> OutBufferRegionsToCopy;
-					EReadResult ReadResult = ReadTiles((uint16*)MipDataPtr, GetMipBufferTotalSize(Dim / MipLevelDiv), Image, FrameId, TileRegions, ConverterParams, CurrentMipLevel, OutBufferRegionsToCopy);
+					EReadResult ReadResult = ReadTiles((uint16*)MipDataPtr, GetMipBufferTotalSize(Dim / MipLevelDiv), Image, TileRegions, ConverterParams, CurrentMipLevel, OutBufferRegionsToCopy);
 					if (ReadResult != Fail)
 					{
 						OutFrame->MipTilesPresent.Emplace(CurrentMipLevel, CurrentTileSelection);
diff --git a/Engine/Plugins/Media/ImgMedia/Source/ImgMedia/Private/Readers/ExrImgMediaReader.h b/Engine/Plugins/Media/ImgMedia/Source/ImgMedia/Private/Readers/ExrImgMediaReader.h
--- a/Engine/Plugins/Media/ImgMedia/Source/ImgMedia/Private/Readers/ExrImgMediaReader.h
+++ b/Engine/Plugins/Media/ImgMedia/Source/ImgMedia/Private/Readers/ExrImgMediaReader.h
@@ -12,6 +12,7 @@
 class FImgMediaLoader;
 class FOpenExrHeaderReader;
 class FRgbaInputFile;
+struct FSampleConverterParameters;
 
 /**
  * Implements a reader for EXR image sequences.
@@ -51,49 +52,6 @@ protected:
 		Skipped
 	};
 
-	/* 
-	* These are all the required parameters to produce Buffer to Texture converter. 
-	*/
-	struct FSampleConverterParameters
-	{
-		/** General file information including its header. */
-		FImgMediaFrameInfo FrameInfo;
-
-		/** Frame Id. */
-		int32 FrameId;
-
-		/** Resolution of the highest quality mip. */
-		FIntPoint FullResolution;
-
-		/** Dimension of the tile including the overscan borders. */
-		FIntPoint TileDimWithBorders;
-
-		/** Used for rendering tiles in bulk regions per mip level. */
-		TMap<int32, TArray<FIntRect>> Viewports;
-
-#if defined(PLATFORM_WINDOWS) && PLATFORM_WINDOWS
-		/** Contain information about individual tiles. Used to convert buffer data into a 2D Texture. 
-		* The size of this array is the exact number of complete and partial tiles for each mip level.
-		*/
-		TArray<TArray<FExrReader::FTileDesc>> TileInfoPerMipLevel;
-#endif
-
-		/** Number of mip levels read. */
-		int32 NumMipLevels;
-
-		/** Pixel stride in bytes. I.e. 2 bytes per pixel x 3 channels = 6. */
-		int32 PixelSize;
-
-		/** Identifies this exr as custom, therefore all data should be swizzled. */
-		bool bCustomExr;
-
-		/** Indicates if mips stored in individual files.*/
-		bool bMipsInSeparateFiles;
-
-		/** A lower quality mip will be upscaled if value is 0 or above. At 0 highest quality mip will always be read fully. */
-		int32 UpscaleMip = -1;
-	};
-
 	/**
 	 * Get the frame information from the given input file.
 	 *
@@ -119,7 +77,6 @@ protected:
 		( uint16* Buffer
 		, int64 BufferSize
 		, const FString& ImagePath
-		, int32 FrameId
 		, const TArray<FIntRect>& TileRegions
 		, TSharedPtr<FSampleConverterParameters> ConverterParams
 		, const int32 CurrentMipLevel
diff --git a/Engine/Plugins/Media/ImgMedia/Source/ImgMedia/Private/Readers/ExrImgMediaReaderGpu.cpp b/Engine/Plugins/Media/ImgMedia/Source/ImgMedia/Private/Readers/ExrImgMediaReaderGpu.cpp
--- a/Engine/Plugins/Media/ImgMedia/Source/ImgMedia/Private/Readers/ExrImgMediaReaderGpu.cpp
+++ b/Engine/Plugins/Media/ImgMedia/Source/ImgMedia/Private/Readers/ExrImgMediaReaderGpu.cpp
@@ -33,6 +33,8 @@ DECLARE_GPU_STAT_NAMED(ExrImgMediaReaderGpu, TEXT("ExrImgGpu"));
 DECLARE_GPU_STAT_NAMED(ExrImgMediaReaderGpu_MipRender, TEXT("ExrImgGpu.MipRender"));
 DECLARE_GPU_STAT_NAMED(ExrImgMediaReaderGpu_MipUpscale, TEXT("ExrImgGpu.MipRender.MipUpscale"));
 DECLARE_GPU_STAT_NAMED(ExrImgMediaReaderGpu_CopyUploadBuffer, TEXT("ExrImgGpu.MipRender.UploadBufferCopy"));
+DECLARE_GPU_STAT_NAMED(ExrImgMediaReaderGpu_AllocateBuffer, TEXT("ExrImgGpu.MipRender.AllocateBuffer"));
+
 
 static TAutoConsoleVariable<bool> CVarExrReaderUseUploadHeap(
 	TEXT("r.ExrReaderGPU.UseUploadHeap"),
@@ -96,53 +98,31 @@ FExrImgMediaReaderGpu::FExrImgMediaReaderGpu(const TSharedRef<FImgMediaLoader, E
 
 FExrImgMediaReaderGpu::~FExrImgMediaReaderGpu()
 {
-	// A signal that tells all buffers that are stored in shared references not to return to the pool
-	// but delete instead.
-	bIsShuttingDown = true;
-
-	// Making sure that all Used memory is processed first and returned into memory pool
-	TransferFromStagingBuffer();
+	FScopeLock ScopeLock(&AllocatorCriticalSecion);
 
-	// Unlock all buffers so that these will release.
-	volatile bool bUnlocked = false;
-	ENQUEUE_RENDER_COMMAND(DeletePooledBuffers)([this, &bUnlocked](FRHICommandListImmediate& RHICmdList)
+	// Copy memory pool array to be released on render thread.
+	ENQUEUE_RENDER_COMMAND(DeletePooledBuffers)([InMemoryPool = MemoryPool](FRHICommandListImmediate& RHICmdList)
 	{
-		FScopeLock ScopeLock(&AllocatorCriticalSecion);
 		SCOPED_DRAW_EVENT(RHICmdList, FExrImgMediaReaderGpu_ReleaseMemoryPool);
 		TArray<uint32> KeysForIteration;
-		MemoryPool.GetKeys(KeysForIteration);
+		InMemoryPool.GetKeys(KeysForIteration);
 		for (uint32 Key : KeysForIteration)
 		{
 			TArray<FStructuredBufferPoolItem*> AllValues;
-			MemoryPool.MultiFind(Key, AllValues);
+			InMemoryPool.MultiFind(Key, AllValues);
 			for (FStructuredBufferPoolItem* MemoryPoolItem : AllValues)
 			{
-				// Check if fence has signaled.
-				check(!MemoryPoolItem->bWillBeSignaled || MemoryPoolItem->RenderFence->Poll());
-				{
-					MemoryPoolItem->RenderFence->Clear();
-					RHIUnlockBuffer(MemoryPoolItem->UploadBufferRef);
-					delete MemoryPoolItem;
-				}
+				delete MemoryPoolItem;
 			}
 		}
-		MemoryPool.Reset();
-		bUnlocked = true;
 	});
-
-	// Wait until unlocking is complete.
-	while (!bUnlocked)
-	{
-		FPlatformProcess::Sleep(0.01f);
-	}
 }
 
 FExrImgMediaReader::EReadResult FExrImgMediaReaderGpu::ReadMip
 	( const int32 CurrentMipLevel
 	, const FImgMediaTileSelection& CurrentTileSelection
 	, TSharedPtr<FImgMediaFrame, ESPMode::ThreadSafe> OutFrame
-	, TSharedPtr<FSampleConverterParameters> ConverterParams
-	, FExrMediaTextureSampleConverter* SampleConverter
+	, TSharedPtr<FExrMediaTextureSampleConverter, ESPMode::ThreadSafe> SampleConverter
 	, const FString& ImagePath
 	, bool bHasTiles)
 {
@@ -151,9 +131,9 @@ FExrImgMediaReader::EReadResult FExrImgMediaReaderGpu::ReadMip
 
 	// Next mip level.
 	int MipLevelDiv = 1 << CurrentMipLevel;
-	FIntPoint CurrentMipDim = ConverterParams->FullResolution / MipLevelDiv;
+	FIntPoint CurrentMipDim = SampleConverter->ConverterParams->FullResolution / MipLevelDiv;
 
-	const SIZE_T BufferSize = GetBufferSize(CurrentMipDim, ConverterParams->FrameInfo.NumChannels, bHasTiles, OutFrame->Info.NumTiles / MipLevelDiv, ConverterParams->bCustomExr);
+	const SIZE_T BufferSize = GetBufferSize(CurrentMipDim, SampleConverter->ConverterParams->FrameInfo.NumChannels, bHasTiles, OutFrame->Info.NumTiles / MipLevelDiv, SampleConverter->ConverterParams->bCustomExr);
 	
 	SampleConverter->LockMipBuffers();
 	FStructuredBufferPoolItemSharedPtr BufferData = SampleConverter->GetMipLevelBuffer(CurrentMipLevel);
@@ -173,7 +153,7 @@ FExrImgMediaReader::EReadResult FExrImgMediaReaderGpu::ReadMip
 	{
 		TArray<UE::Math::TIntPoint<int64>> BufferRegionsToCopy;
 		// read frame data
-		if (bHasTiles || ConverterParams->bCustomExr)
+		if (bHasTiles || SampleConverter->ConverterParams->bCustomExr)
 		{
 			TArray<FIntRect> TileRegionsToRead;
 			{
@@ -196,7 +176,7 @@ FExrImgMediaReader::EReadResult FExrImgMediaReaderGpu::ReadMip
 			}
 			else
 			{
-				ReadResult = ReadTiles(MipDataPtr, BufferSize, ImagePath, ConverterParams->FrameId, TileRegionsToRead, ConverterParams, CurrentMipLevel, BufferRegionsToCopy);
+				ReadResult = ReadTiles(MipDataPtr, BufferSize, ImagePath, TileRegionsToRead, SampleConverter->ConverterParams, CurrentMipLevel, BufferRegionsToCopy);
 				
 				for (const FIntRect& Region : TileRegionsToRead)
 				{
@@ -206,16 +186,16 @@ FExrImgMediaReader::EReadResult FExrImgMediaReaderGpu::ReadMip
 		}
 		else
 		{
-			ReadResult = ReadInChunks(MipDataPtr, ImagePath, ConverterParams->FrameId, CurrentMipDim, BufferSize);
+			ReadResult = ReadInChunks(MipDataPtr, ImagePath, SampleConverter->ConverterParams->FrameId, CurrentMipDim, BufferSize);
 			OutFrame->NumTilesRead++;
 		}
 
 		if (ReadResult == Success && bUseUploadHeap)
 		{
-			ENQUEUE_RENDER_COMMAND(CopyFromUploadBuffer)([SampleConverter, ConverterParams, BufferData, BufferRegionsToCopy](FRHICommandListImmediate& RHICmdList)
+			ENQUEUE_RENDER_COMMAND(CopyFromUploadBuffer)([SampleConverter, BufferData, BufferRegionsToCopy](FRHICommandListImmediate& RHICmdList)
 				{
 					SCOPED_DRAW_EVENT(RHICmdList, FExrImgMediaReaderGpu_CopyBuffers);
-					TRACE_CPUPROFILER_EVENT_SCOPE_TEXT(*FString::Printf(TEXT("ExrReaderGpu.StartCopy %d"), ConverterParams->FrameId));
+					TRACE_CPUPROFILER_EVENT_SCOPE_TEXT(*FString::Printf(TEXT("ExrReaderGpu.StartCopy %d"), SampleConverter->ConverterParams->FrameId));
 					SCOPED_GPU_STAT(RHICmdList, ExrImgMediaReaderGpu_CopyUploadBuffer);
 
 					if (BufferRegionsToCopy.IsEmpty())
@@ -266,10 +246,11 @@ bool FExrImgMediaReaderGpu::ReadFrame(int32 FrameId, const TMap<int32, FImgMedia
 		return false;
 	}
 
+	TSharedPtr<FExrMediaTextureSampleConverter, ESPMode::ThreadSafe> SampleConverterPtr = OutFrame->GetOrCreateSampleConverter<FExrMediaTextureSampleConverter>();
+
 	// Get tile info.
 	bool bHasTiles = OutFrame->Info.bHasTiles;
-
-	TSharedPtr<FSampleConverterParameters> ConverterParams = MakeShared<FSampleConverterParameters>();
+	TSharedPtr<FSampleConverterParameters> ConverterParams = SampleConverterPtr->ConverterParams;
 	ConverterParams->FullResolution = OutFrame->Info.Dim;
 	ConverterParams->FrameId = FrameId;
 
@@ -287,14 +268,6 @@ bool FExrImgMediaReaderGpu::ReadFrame(int32 FrameId, const TMap<int32, FImgMedia
 	ConverterParams->bCustomExr = OutFrame->Info.FormatName == TEXT("EXR CUSTOM");
 	ConverterParams->bMipsInSeparateFiles = Loader->MipsInSeparateFiles();
 
-	FExrMediaTextureSampleConverter* SampleConverter;
-	if (!OutFrame->SampleConverter.IsValid())
-	{
-		OutFrame->SampleConverter = MakeShared<FExrMediaTextureSampleConverter>();
-	}
-	SampleConverter = static_cast<FExrMediaTextureSampleConverter*>(OutFrame->SampleConverter.Get());
-
-
 	{
 		// Force mip level to be upscaled to all higher quality mips.
 		FIntPoint CurrentMipDim = ConverterParams->FullResolution;
@@ -330,7 +303,7 @@ bool FExrImgMediaReaderGpu::ReadFrame(int32 FrameId, const TMap<int32, FImgMedia
 			// Get highest resolution mip level path.
 			FString ImagePath = Loader->GetImagePath(ConverterParams->FrameId, ConverterParams->bMipsInSeparateFiles ? CurrentMipLevel : 0);
 
-			EReadResult ReadResult = ReadMip(CurrentMipLevel, CurrentTileSelection, OutFrame, ConverterParams, SampleConverter, ImagePath, bHasTiles);
+			EReadResult ReadResult = ReadMip(CurrentMipLevel, CurrentTileSelection, OutFrame, SampleConverterPtr, ImagePath, bHasTiles);
 			switch (ReadResult)
 			{
 			case FExrImgMediaReader::Success:
@@ -359,9 +332,6 @@ bool FExrImgMediaReaderGpu::ReadFrame(int32 FrameId, const TMap<int32, FImgMedia
 					// Fall back to CPU.
 					bFallBackToCPU = true;
 
-					// To make sure that Media Texture doesn't call the converter if this frame is invalid.
-					OutFrame->SampleConverter.Reset();
-
 					return FExrImgMediaReader::ReadFrame(ConverterParams->FrameId, InMipTiles, OutFrame);
 				}
 				break;
@@ -385,7 +355,7 @@ bool FExrImgMediaReaderGpu::ReadFrame(int32 FrameId, const TMap<int32, FImgMedia
 			const int32 CurrentMipLevel = TilesPerMip.Key;
 
 			// Skip this viewport since we don't have anything to render.
-			if (!SampleConverter->GetMipLevelBuffer(CurrentMipLevel))
+			if (!SampleConverterPtr->GetMipLevelBuffer(CurrentMipLevel))
 			{
 				continue;
 			}
@@ -416,7 +386,7 @@ bool FExrImgMediaReaderGpu::ReadFrame(int32 FrameId, const TMap<int32, FImgMedia
 	OutFrame->Stride = ConverterParams->FullResolution.X * ConverterParams->PixelSize;
 	
 	ConverterParams->Viewports.KeySort(TLess<int32>());
-	CreateSampleConverterCallback(SampleConverter, ConverterParams);
+	CreateSampleConverterCallback(SampleConverterPtr);
 
 	UE_LOG(LogImgMedia, Verbose, TEXT("Reader %p: Read Pixels Complete. %i"), this, FrameId);
 	return true;
@@ -424,6 +394,7 @@ bool FExrImgMediaReaderGpu::ReadFrame(int32 FrameId, const TMap<int32, FImgMedia
 
 void FExrImgMediaReaderGpu::PreAllocateMemoryPool(int32 NumFrames, const FImgMediaFrameInfo& FrameInfo, const bool bCustomExr)
 {
+	TRACE_CPUPROFILER_EVENT_SCOPE_TEXT(*FString::Printf(TEXT("ExrReaderGpu.PreAllocateMemoryPool")));
 	SIZE_T AllocSize = GetBufferSize(FrameInfo.Dim, FrameInfo.NumChannels, FrameInfo.bHasTiles, FrameInfo.NumTiles, bCustomExr);
 	for (int32 FrameCacheNum = 0; FrameCacheNum < NumFrames; FrameCacheNum++)
 	{
@@ -431,18 +402,6 @@ void FExrImgMediaReaderGpu::PreAllocateMemoryPool(int32 NumFrames, const FImgMed
 	}
 }
 
-
-void FExrImgMediaReaderGpu::OnTick()
-{
-	// Only tick once per frame.
-	if (LastTickedFrameCounter != GFrameCounter)
-	{
-		LastTickedFrameCounter = GFrameCounter;
-
-		TransferFromStagingBuffer();
-	}
-}
-
 /* FExrImgMediaReaderGpu implementation
  *****************************************************************************/
 
@@ -530,20 +489,18 @@ SIZE_T FExrImgMediaReaderGpu::GetBufferSize(const FIntPoint& Dim, int32 NumChann
 
 }
 
-void FExrImgMediaReaderGpu::CreateSampleConverterCallback
-	( FExrMediaTextureSampleConverter* SampleConverter
-	, TSharedPtr<FSampleConverterParameters> ConverterParams)
+void FExrImgMediaReaderGpu::CreateSampleConverterCallback(TSharedPtr<FExrMediaTextureSampleConverter, ESPMode::ThreadSafe> SampleConverter)
 {
-	auto RenderThreadSwizzler = [ConverterParams] (FRHICommandListImmediate& RHICmdList, FTexture2DRHIRef RenderTargetTextureRHI, TMap<int32, FStructuredBufferPoolItemSharedPtr>& MipBuffers)->bool
+	auto RenderThreadSwizzler = [] (FRHICommandListImmediate& RHICmdList, FTexture2DRHIRef RenderTargetTextureRHI, TMap<int32, FStructuredBufferPoolItemSharedPtr>& MipBuffers, const FSampleConverterParameters& ConverterParams)->bool
 	{
-		TRACE_CPUPROFILER_EVENT_SCOPE_TEXT(*FString::Printf(TEXT("ExrReaderGpu.Convert %d"), ConverterParams->FrameId));
+		TRACE_CPUPROFILER_EVENT_SCOPE_TEXT(*FString::Printf(TEXT("ExrReaderGpu.Convert %d"), ConverterParams.FrameId));
 		SCOPED_DRAW_EVENT(RHICmdList, FExrImgMediaReaderGpu_Convert);
 		SCOPED_GPU_STAT(RHICmdList, ExrImgMediaReaderGpu);
 
 		auto RenderMip = []
 			( FRHICommandListImmediate& RHICmdList
 			, FTexture2DRHIRef RenderTargetTextureRHI
-			, TSharedPtr<FSampleConverterParameters> ConverterParams
+			, const FSampleConverterParameters & ConverterParams
 			, int32 SampleMipLevel
 			, int32 TextureMipLevel
 			, FStructuredBufferPoolItemSharedPtr BufferData
@@ -557,34 +514,34 @@ void FExrImgMediaReaderGpu::CreateSampleConverterCallback
 			RHICmdList.BeginRenderPass(RPInfo, TEXT("ExrTextureSwizzle"));
 
 			FExrSwizzlePS::FPermutationDomain PermutationVector;
-			PermutationVector.Set<FExrSwizzlePS::FRgbaSwizzle>(ConverterParams->FrameInfo.NumChannels - 1);
-			PermutationVector.Set<FExrSwizzlePS::FRenderTiles>(ConverterParams->FrameInfo.bHasTiles || ConverterParams->bCustomExr);
-			PermutationVector.Set<FExrSwizzlePS::FCustomExr>(ConverterParams->bCustomExr);
+			PermutationVector.Set<FExrSwizzlePS::FRgbaSwizzle>(ConverterParams.FrameInfo.NumChannels - 1);
+			PermutationVector.Set<FExrSwizzlePS::FRenderTiles>(ConverterParams.FrameInfo.bHasTiles || ConverterParams.bCustomExr);
+			PermutationVector.Set<FExrSwizzlePS::FCustomExr>(ConverterParams.bCustomExr);
 			PermutationVector.Set<FExrSwizzlePS::FPartialTiles>(false);
 
 			FExrSwizzlePS::FParameters Parameters = FExrSwizzlePS::FParameters();
 			Parameters.TextureSize = SampleSize;
-			Parameters.TileSize = ConverterParams->TileDimWithBorders;
-			Parameters.NumChannels = ConverterParams->FrameInfo.NumChannels;
-			if (ConverterParams->FrameInfo.bHasTiles)
+			Parameters.TileSize = ConverterParams.TileDimWithBorders;
+			Parameters.NumChannels = ConverterParams.FrameInfo.NumChannels;
+			if (ConverterParams.FrameInfo.bHasTiles)
 			{
-				Parameters.NumTiles = FIntPoint(FMath::CeilToInt(float(SampleSize.X) / ConverterParams->TileDimWithBorders.X), FMath::CeilToInt(float(SampleSize.Y) / ConverterParams->TileDimWithBorders.Y));
+				Parameters.NumTiles = FIntPoint(FMath::CeilToInt(float(SampleSize.X) / ConverterParams.TileDimWithBorders.X), FMath::CeilToInt(float(SampleSize.Y) / ConverterParams.TileDimWithBorders.Y));
 			}
 
-			if (ConverterParams->FrameInfo.bHasTiles &&
-				(ConverterParams->TileInfoPerMipLevel.Num() > SampleMipLevel && ConverterParams->TileInfoPerMipLevel[SampleMipLevel].Num() > 0))
+			if (ConverterParams.FrameInfo.bHasTiles &&
+				(ConverterParams.TileInfoPerMipLevel.Num() > SampleMipLevel && ConverterParams.TileInfoPerMipLevel[SampleMipLevel].Num() > 0))
 			{
 				TRACE_CPUPROFILER_EVENT_SCOPE_TEXT(*FString::Printf(TEXT("ExrReaderGpu.TileDesc")));
 
 				FBufferRHIRef BufferRef;
 				FRHIResourceCreateInfo CreateInfo(TEXT("FExrImgMediaReaderGpu_TileDesc"));
 				const uint32 BytesPerElement = sizeof(FExrReader::FTileDesc);
-				const uint32 NumElements = ConverterParams->TileInfoPerMipLevel[SampleMipLevel].Num();
+				const uint32 NumElements = ConverterParams.TileInfoPerMipLevel[SampleMipLevel].Num();
 
 				// This buffer is allocated on already allocated block, therefore the risk of fragmentation is mitigated.
 				BufferRef = RHICreateStructuredBuffer(BytesPerElement, BytesPerElement * NumElements, BUF_ShaderResource | BUF_Dynamic | BUF_FastVRAM, CreateInfo);
 				void* MappedBuffer = RHILockBuffer(BufferRef, 0, NumElements, RLM_WriteOnly);
-				FMemory::Memcpy(MappedBuffer, &ConverterParams->TileInfoPerMipLevel[SampleMipLevel][0], sizeof(FExrReader::FTileDesc) * ConverterParams->TileInfoPerMipLevel[SampleMipLevel].Num());
+				FMemory::Memcpy(MappedBuffer, &ConverterParams.TileInfoPerMipLevel[SampleMipLevel][0], sizeof(FExrReader::FTileDesc) * ConverterParams.TileInfoPerMipLevel[SampleMipLevel].Num());
 				RHIUnlockBuffer(BufferRef);
 				Parameters.TileDescBuffer = RHICreateShaderResourceView(BufferRef);
 				PermutationVector.Set<FExrSwizzlePS::FPartialTiles>(true);
@@ -610,24 +567,15 @@ void FExrImgMediaReaderGpu::CreateSampleConverterCallback
 
 			// Resolve render target.
 			RHICmdList.EndRenderPass();
-
-			if (!BufferData->bWillBeSignaled)
-			{
-				// This flag will indicate that we should wait for poll to complete.
-				BufferData->bWillBeSignaled = true;
-
-				// Mark this render command for this buffer as complete, so we can poll it and transfer later.
-				RHICmdList.WriteGPUFence(BufferData->RenderFence);
-			}
 		};
 
-		int32 MipToUpscale = ConverterParams->UpscaleMip;
+		int32 MipToUpscale = ConverterParams.UpscaleMip;
 
 		// Upscale to all mips below the mip to upscale.
 		for (int32 MipLevel = 0; MipLevel <= MipToUpscale; MipLevel++)
 		{
 			int MipLevelDiv = 1 << MipLevel;
-			FIntPoint Dim = ConverterParams->FullResolution / MipLevelDiv;
+			FIntPoint Dim = ConverterParams.FullResolution / MipLevelDiv;
 
 			{
 				SCOPED_GPU_STAT(RHICmdList, ExrImgMediaReaderGpu_MipUpscale);
@@ -641,11 +589,11 @@ void FExrImgMediaReaderGpu::CreateSampleConverterCallback
 				FStructuredBufferPoolItemSharedPtr BufferDataToUpscale = MipBuffers[MipToUpscale];
 				TArray<FIntRect> FakeViewport;
 				FakeViewport.Add(FIntRect(FIntPoint(0, 0), Dim));
-				RenderMip(RHICmdList, RenderTargetTextureRHI, ConverterParams, MipToUpscale, MipLevel, BufferDataToUpscale, (ConverterParams->FullResolution / (1 << MipToUpscale)), Dim, FakeViewport);
+				RenderMip(RHICmdList, RenderTargetTextureRHI, ConverterParams, MipToUpscale, MipLevel, BufferDataToUpscale, (ConverterParams.FullResolution / (1 << MipToUpscale)), Dim, FakeViewport);
 			}
 		}
 
-		for (const TPair<int32, TArray<FIntRect>>& MipLevelViewports : ConverterParams->Viewports)
+		for (const TPair<int32, TArray<FIntRect>>& MipLevelViewports : ConverterParams.Viewports)
 		{
 			int32 MipLevel = MipLevelViewports.Key;
 			
@@ -657,7 +605,7 @@ void FExrImgMediaReaderGpu::CreateSampleConverterCallback
 
 			FStructuredBufferPoolItemSharedPtr BufferData = MipBuffers[MipLevel];
 			int MipLevelDiv = 1 << MipLevel;
-			FIntPoint Dim = ConverterParams->FullResolution / MipLevelDiv;
+			FIntPoint Dim = ConverterParams.FullResolution / MipLevelDiv;
 
 			if (BufferData.IsValid())
 			{
@@ -685,11 +633,24 @@ void FExrImgMediaReaderGpu::CreateSampleConverterCallback
 
 FStructuredBufferPoolItemSharedPtr FExrImgMediaReaderGpu::AllocateGpuBufferFromPool(uint32 AllocSize, bool bWait)
 {
+	TRACE_CPUPROFILER_EVENT_SCOPE_TEXT(*FString::Printf(TEXT("ExrReaderGpu.AllocBuffer")));
 	TRACE_CPUPROFILER_EVENT_SCOPE_TEXT(*FString::Printf(TEXT("ExrReaderGpu.AllocBuffer %d"), AllocSize));
+	TWeakPtr<FExrImgMediaReaderGpu, ESPMode::ThreadSafe> WeakReaderPtr = AsShared();
 
 	// This function is attached to the shared pointer and is used to return any allocated memory to staging pool.
-	auto BufferDeleter = [AllocSize](FStructuredBufferPoolItem* ObjectToDelete) {
-		ReturnGpuBufferToStagingPool(AllocSize, ObjectToDelete);
+	auto BufferDeleter = [WeakReaderPtr, AllocSize](FStructuredBufferPoolItem* ObjectToDelete) {
+		TSharedPtr<FExrImgMediaReaderGpu, ESPMode::ThreadSafe> SharedReaderPtr = WeakReaderPtr.Pin();
+		if (SharedReaderPtr.IsValid())
+		{
+			SharedReaderPtr->ReturnGpuBufferToPool(AllocSize, ObjectToDelete);
+		}
+		else
+		{
+			ENQUEUE_RENDER_COMMAND(DeletePooledBuffers)([ObjectToDelete](FRHICommandListImmediate& RHICmdList)
+			{
+				delete ObjectToDelete;
+			});
+		}
 	};
 
 	// Buffer that ends up being returned out of this function.
@@ -716,7 +677,10 @@ FStructuredBufferPoolItemSharedPtr FExrImgMediaReaderGpu::AllocateGpuBufferFromP
 			// Allocate and unlock the structured buffer on render thread.
 			ENQUEUE_RENDER_COMMAND(CreatePooledBuffer)([AllocatedBuffer, AllocSize, &bInitDone, this, bWait](FRHICommandListImmediate& RHICmdList)
 			{
+				TRACE_CPUPROFILER_EVENT_SCOPE_TEXT(*FString::Printf(TEXT("ExrReaderGpu.AllocBuffer_RenderThread")));
+				TRACE_CPUPROFILER_EVENT_SCOPE_TEXT(*FString::Printf(TEXT("ExrReaderGpu.AllocBuffer_RenderThread %d"), AllocSize));
 				FScopeLock ScopeLock(&AllocatorCriticalSecion);
+				SCOPED_GPU_STAT(RHICmdList, ExrImgMediaReaderGpu_AllocateBuffer);
 				SCOPED_DRAW_EVENT(RHICmdList, FExrImgMediaReaderGpu_AllocateBuffer);
 				FRHIResourceCreateInfo CreateInfo(TEXT("FExrImgMediaReaderGpu"));
 				AllocatedBuffer->UploadBufferRef = RHICreateStructuredBuffer(sizeof(uint16) * 2., AllocSize, BUF_ShaderResource | BUF_Dynamic | BUF_FastVRAM, CreateInfo);
@@ -732,8 +696,6 @@ FStructuredBufferPoolItemSharedPtr FExrImgMediaReaderGpu::AllocateGpuBufferFromP
 					AllocatedBuffer->ShaderResourceView = RHICreateShaderResourceView(AllocatedBuffer->UploadBufferRef);
 				}
 
-				AllocatedBuffer->RenderFence = RHICreateGPUFence(TEXT("Exr.BufferNoLongerInUseFence"));
-
 				if (bWait)
 				{
 					bInitDone = true;
@@ -753,61 +715,11 @@ FStructuredBufferPoolItemSharedPtr FExrImgMediaReaderGpu::AllocateGpuBufferFromP
 	return AllocatedBuffer;
 }
 
-void FExrImgMediaReaderGpu::ReturnGpuBufferToStagingPool(uint32 AllocSize, FStructuredBufferPoolItem* Buffer)
+void FExrImgMediaReaderGpu::ReturnGpuBufferToPool(uint32 AllocSize, FStructuredBufferPoolItem* Buffer)
 {
-	TSharedPtr< FExrImgMediaReaderGpu, ESPMode::ThreadSafe> Reader = Buffer->Reader.Pin();
-
-	// If reader is being deleted, we don't need to return the memory into staging buffer and instead should delete it.
-	if ((Reader.IsValid() == false) || (Reader->bIsShuttingDown))
-	{
-		ENQUEUE_RENDER_COMMAND(DeletePooledBuffers)([Buffer](FRHICommandListImmediate& RHICmdList)
-		{
-			SCOPED_DRAW_EVENT(RHICmdList, FExrImgMediaReaderGpu_ReleaseBuffer);
-			Buffer->RenderFence->Clear();
-			// By this point we don't need a lock because the destructor was already called and it 
-			// is guaranteed that this buffer is no longer used anywhere else.
-			RHIUnlockBuffer(Buffer->UploadBufferRef);
-			delete Buffer;
-		});
-	}
-	else
-	{
-		FScopeLock ScopeLock(&Reader->AllocatorCriticalSecion);
-
-		// We don't need to process this pooled buffer if the Reader is being destroyed.
-		Reader->StagingMemoryPool.Add(AllocSize, Buffer);
-	}
-
-}
-
-void FExrImgMediaReaderGpu::TransferFromStagingBuffer()
-{
-	ENQUEUE_RENDER_COMMAND(CreatePooledBuffer)([&, this](FRHICommandListImmediate& RHICmdList)
-	{
-		FScopeLock ScopeLock(&AllocatorCriticalSecion);
-		SCOPED_DRAW_EVENT(RHICmdList, FExrImgMediaReaderGpu_TransferFromStagingBuffer);
-
-		TArray<uint32> KeysForIteration;
-		StagingMemoryPool.GetKeys(KeysForIteration);
-		for (uint32 Key : KeysForIteration)
-		{
-			TArray<FStructuredBufferPoolItem*> AllValues;
-			StagingMemoryPool.MultiFind(Key, AllValues);
-			for (FStructuredBufferPoolItem* MemoryPoolItem : AllValues)
-			{
-				// Check if fence has signaled. Or otherwise if we are waiting for signal to come through.
-				if (!MemoryPoolItem->bWillBeSignaled || MemoryPoolItem->RenderFence->Poll())
-				{
-					// If buffer was in use but the fence signaled we need to reset bWillBeSignaled flag.
-					MemoryPoolItem->bWillBeSignaled = false;
-					MemoryPoolItem->RenderFence->Clear();
-
-					StagingMemoryPool.Remove(Key, MemoryPoolItem);
-					MemoryPool.Add(Key, MemoryPoolItem);
-				}
-			}
-		}
-	});
+	FScopeLock ScopeLock(&AllocatorCriticalSecion);
+	TRACE_CPUPROFILER_EVENT_SCOPE_TEXT(*FString::Printf(TEXT("ExrReaderGpu.ReturnPoolItem")));
+	MemoryPool.Add(AllocSize, Buffer);
 }
 
 
@@ -820,10 +732,21 @@ bool FExrMediaTextureSampleConverter::Convert(FTexture2DRHIRef& InDstTexture, co
 	bool bExecutionSuccessful = false;
 	if (ConvertExrBufferCallback.IsBound())
 	{
-		bExecutionSuccessful = ConvertExrBufferCallback.Execute(FRHICommandListExecutor::GetImmediateCommandList(), InDstTexture, MipBuffersRenderThread);
+		bExecutionSuccessful = ConvertExrBufferCallback.Execute(FRHICommandListExecutor::GetImmediateCommandList(), InDstTexture, MipBuffers, *ConverterParams.Get());
 	}
 	return bExecutionSuccessful;
 }
 
+FStructuredBufferPoolItem::~FStructuredBufferPoolItem()
+{
+	TRACE_CPUPROFILER_EVENT_SCOPE_TEXT(*FString::Printf(TEXT("ExrReaderGpu.ReleasePoolItem")));
+	check(IsInRenderingThread());
+	{
+		RHIUnlockBuffer(UploadBufferRef);
+		UploadBufferMapped = nullptr;
+	}
+
+}
+
 #endif //IMGMEDIA_EXR_SUPPORTED_PLATFORM && PLATFORM_WINDOWS
 
diff --git a/Engine/Plugins/Media/ImgMedia/Source/ImgMedia/Private/Readers/ExrImgMediaReaderGpu.h b/Engine/Plugins/Media/ImgMedia/Source/ImgMedia/Private/Readers/ExrImgMediaReaderGpu.h
--- a/Engine/Plugins/Media/ImgMedia/Source/ImgMedia/Private/Readers/ExrImgMediaReaderGpu.h
+++ b/Engine/Plugins/Media/ImgMedia/Source/ImgMedia/Private/Readers/ExrImgMediaReaderGpu.h
@@ -35,21 +35,15 @@ struct FStructuredBufferPoolItem
 	*/
 	FShaderResourceViewRHIRef ShaderResourceView;
 
-	/** 
-	* A Gpu fence that identifies if this pool item is available for use again.
-	*/
-	FGPUFenceRHIRef RenderFence;
-
 	/**
-	* This boolean is used as a flag in combination with fences to indicate if rendering thread 
-	* is currently using it.
+	* Keep track of our reader in case it gets destroyed.
 	*/
-	bool bWillBeSignaled = false;
+	TWeakPtr<FExrImgMediaReaderGpu, ESPMode::ThreadSafe> Reader;
 
 	/**
-	* Keep track of our reader in case it gets destroyed.
+	* Destructure to clean up render resources
 	*/
-	TWeakPtr<FExrImgMediaReaderGpu, ESPMode::ThreadSafe> Reader;
+	~FStructuredBufferPoolItem();
 };
 
 /**
@@ -79,7 +73,6 @@ public:
 	* For performance reasons we want to pre-allocate structured buffers to at least the number of concurrent frames.
 	*/
 	virtual void PreAllocateMemoryPool(int32 NumFrames, const FImgMediaFrameInfo& FrameInfo, const bool bCustomExr) override;
-	virtual void OnTick() override;
 
 protected:
 
@@ -100,8 +93,7 @@ protected:
 	/**
 	* Creates Sample converter to be used by Media Texture Resource.
 	*/
-	static void CreateSampleConverterCallback
-		(FExrMediaTextureSampleConverter* SampleConverter, TSharedPtr<FSampleConverterParameters> ConverterParams);
+	static void CreateSampleConverterCallback(TSharedPtr<FExrMediaTextureSampleConverter, ESPMode::ThreadSafe> SampleConverter);
 
 	/**
 	* A function that reads one mip level.
@@ -110,36 +102,28 @@ protected:
 		( const int32 CurrentMipLevel
 		, const FImgMediaTileSelection& CurrentTileSelection
 		, TSharedPtr<FImgMediaFrame, ESPMode::ThreadSafe> OutFrame
-		, TSharedPtr<FSampleConverterParameters> ConverterParams
-		, FExrMediaTextureSampleConverter* SampleConverter
+		, TSharedPtr<FExrMediaTextureSampleConverter, ESPMode::ThreadSafe> SampleConverter
 		, const FString& ImagePath
 		, bool bHasTiles);
 
 public:
 
 	/** Typically we would need the (ImageResolution.x*y)*NumChannels*ChannelSize */
-	FStructuredBufferPoolItemSharedPtr AllocateGpuBufferFromPool(uint32 AllocSize, bool bWait = true);
+	virtual FStructuredBufferPoolItemSharedPtr AllocateGpuBufferFromPool(uint32 AllocSize, bool bWait = true);
 
 	/** Either return or Add new chunk of memory to the pool based on its size. */
-	static void ReturnGpuBufferToStagingPool(uint32 AllocSize, FStructuredBufferPoolItem* Buffer);
+	void ReturnGpuBufferToPool(uint32 AllocSize, FStructuredBufferPoolItem* Buffer);
 
-	/** Transfer from Staging buffer to Memory pool. */
-	void TransferFromStagingBuffer();
+	
 
 private:
 
 	/** A critical section used for memory allocation and pool management. */
 	FCriticalSection AllocatorCriticalSecion;
 
-	/** Main memory pool from where we are allowed to take buffers. */
+	/** Memory pool from where we are allowed to take buffers. */
 	TMultiMap<uint32, FStructuredBufferPoolItem*> MemoryPool;
 
-	/** 
-	* This pool could contain potentially in use buffers and every tick it is processed
-	* and those buffers that are ready to be used returned back to Main memory pool
-	*/
-	TMultiMap<uint32, FStructuredBufferPoolItem*> StagingMemoryPool;
-
 	/** Frame that was last ticked so we don't tick more than once. */
 	uint64 LastTickedFrameCounter;
 
@@ -150,7 +134,52 @@ private:
 	bool bFallBackToCPU;
 };
 
-FUNC_DECLARE_DELEGATE(FExrConvertBufferCallback, bool, FRHICommandListImmediate& /*RHICmdList*/, FTexture2DRHIRef /*RenderTargetTextureRHI*/, TMap<int32, FStructuredBufferPoolItemSharedPtr>& /*MipBuffers*/ )
+
+/*
+* These are all the required parameters to convert Buffer to Texture converter.
+*/
+struct FSampleConverterParameters
+{
+	/** General file information including its header. */
+	FImgMediaFrameInfo FrameInfo;
+
+	/** Frame Id. */
+	int32 FrameId;
+
+	/** Resolution of the highest quality mip. */
+	FIntPoint FullResolution;
+
+	/** Dimension of the tile including the overscan borders. */
+	FIntPoint TileDimWithBorders;
+
+	/** Used for rendering tiles in bulk regions per mip level. */
+	TMap<int32, TArray<FIntRect>> Viewports;
+
+#if defined(PLATFORM_WINDOWS) && PLATFORM_WINDOWS
+	/** Contain information about individual tiles. Used to convert buffer data into a 2D Texture.
+	* The size of this array is the exact number of complete and partial tiles for each mip level.
+	*/
+	TArray<TArray<FExrReader::FTileDesc>> TileInfoPerMipLevel;
+#endif
+
+	/** Number of mip levels read. */
+	int32 NumMipLevels;
+
+	/** Pixel stride in bytes. I.e. 2 bytes per pixel x 3 channels = 6. */
+	int32 PixelSize;
+
+	/** Identifies this exr as custom, therefore all data should be swizzled. */
+	bool bCustomExr;
+
+	/** Indicates if mips stored in individual files.*/
+	bool bMipsInSeparateFiles;
+
+	/** A lower quality mip will be upscaled if value is 0 or above. At 0 highest quality mip will always be read fully. */
+	int32 UpscaleMip = -1;
+};
+
+
+FUNC_DECLARE_DELEGATE(FExrConvertBufferCallback, bool, FRHICommandListImmediate& /*RHICmdList*/, FTexture2DRHIRef /*RenderTargetTextureRHI*/, TMap<int32, FStructuredBufferPoolItemSharedPtr>& /*MipBuffers*/, const FSampleConverterParameters & /*ConverterParams*/)
 
 class FExrMediaTextureSampleConverter: public IMediaTextureSampleConverter
 {
@@ -193,6 +222,9 @@ public:
 		check(!MipBuffers.Contains(RequestedMipLevel));
 		MipBuffers.Add(RequestedMipLevel, Buffer);
 	}
+public:
+	/** These are all required parameters to convert the buffer into texture successfully. */
+	TSharedPtr<FSampleConverterParameters> ConverterParams = MakeShared<FSampleConverterParameters>();
 
 private:
 	FCriticalSection ConverterCallbacksCriticalSection;
diff --git a/Engine/Plugins/Media/ImgMedia/Source/ImgMedia/Private/Readers/IImgMediaReader.h b/Engine/Plugins/Media/ImgMedia/Source/ImgMedia/Private/Readers/IImgMediaReader.h
--- a/Engine/Plugins/Media/ImgMedia/Source/ImgMedia/Private/Readers/IImgMediaReader.h
+++ b/Engine/Plugins/Media/ImgMedia/Source/ImgMedia/Private/Readers/IImgMediaReader.h
@@ -77,9 +77,7 @@ struct FImgMediaFrame
 	/** The frame's horizontal stride (in bytes). */
 	uint32 Stride = 0;
 
-	/** Sample converter is used by Media Texture Resource to convert the texture or data. */
-	TSharedPtr<IMediaTextureSampleConverter, ESPMode::ThreadSafe> SampleConverter;
-
+	// This should only be used if you do not need to make changes to sample converter.
 	virtual IMediaTextureSampleConverter* GetSampleConverter()
 	{
 		if (!SampleConverter.IsValid())
@@ -91,6 +89,24 @@ struct FImgMediaFrame
 
 	/** Virtual non trivial destructor. */
 	virtual ~FImgMediaFrame() {};
+
+	template <class T>
+	TSharedPtr<T, ESPMode::ThreadSafe> GetOrCreateSampleConverter()
+	{
+		FScopeLock ScopedLock(SampleConverterCriticalSection.Get());
+		if (!SampleConverter.IsValid())
+		{
+			SampleConverter = MakeShared<T>();
+		}
+		return StaticCastSharedPtr<T>(SampleConverter);
+	}
+
+private:
+	/** Lock to be used exclusively on reader threads.*/
+	TSharedPtr<FCriticalSection, ESPMode::ThreadSafe> SampleConverterCriticalSection = MakeShared<FCriticalSection>();
+
+	/** Sample converter is used by Media Texture Resource to convert the texture or data. */
+	TSharedPtr<IMediaTextureSampleConverter, ESPMode::ThreadSafe> SampleConverter;
 };
 
 
