ImgMedia: Minor EXRGpuReader cleanup removing direct link to DX12, replacing busy loop wait with FEvent.

diff --git a/Engine/Plugins/Media/ImgMedia/Source/ImgMedia/ImgMedia.Build.cs b/Engine/Plugins/Media/ImgMedia/Source/ImgMedia/ImgMedia.Build.cs
--- a/Engine/Plugins/Media/ImgMedia/Source/ImgMedia/ImgMedia.Build.cs
+++ b/Engine/Plugins/Media/ImgMedia/Source/ImgMedia/ImgMedia.Build.cs
@@ -49,12 +49,6 @@ namespace UnrealBuildTool.Rules
 				PrivateDependencyModuleNames.Add("OpenExrWrapper");
 				PrivateDependencyModuleNames.Add("ExrReaderGpu");
 			}
-
-			if (Target.Platform.IsInGroup(UnrealPlatformGroup.Windows))
-			{
-				AddEngineThirdPartyPrivateStaticDependencies(Target, "DX12");
-				PrivateDependencyModuleNames.Add("D3D12RHI");
-			}
 			
 			// Is this the editor?
 			if (Target.bBuildEditor == true)
diff --git a/Engine/Plugins/Media/ImgMedia/Source/ImgMedia/Private/Readers/ExrImgMediaReader.cpp b/Engine/Plugins/Media/ImgMedia/Source/ImgMedia/Private/Readers/ExrImgMediaReader.cpp
--- a/Engine/Plugins/Media/ImgMedia/Source/ImgMedia/Private/Readers/ExrImgMediaReader.cpp
+++ b/Engine/Plugins/Media/ImgMedia/Source/ImgMedia/Private/Readers/ExrImgMediaReader.cpp
@@ -362,6 +362,7 @@ TSharedPtr<IImgMediaReader, ESPMode::ThreadSafe> FExrImgMediaReader::GetReader(c
 	}
 
 	// Check GetCompressionName of OpenExrWrapper for other compression names.
+	// todo: Add and test Vulkan support
 	if (GDynamicRHI && GDynamicRHI->GetInterfaceType() == ERHIInterfaceType::D3D12
 		&& Info.CompressionName == "Uncompressed" 
 		&& CVarEnableUncompressedExrGpuReader.GetValueOnAnyThread()
diff --git a/Engine/Plugins/Media/ImgMedia/Source/ImgMedia/Private/Readers/ExrImgMediaReaderGpu.cpp b/Engine/Plugins/Media/ImgMedia/Source/ImgMedia/Private/Readers/ExrImgMediaReaderGpu.cpp
--- a/Engine/Plugins/Media/ImgMedia/Source/ImgMedia/Private/Readers/ExrImgMediaReaderGpu.cpp
+++ b/Engine/Plugins/Media/ImgMedia/Source/ImgMedia/Private/Readers/ExrImgMediaReaderGpu.cpp
@@ -26,8 +26,6 @@
 #include "TextureResource.h"
 #include "ScreenPass.h"
 
-#include "ID3D12DynamicRHI.h"
-
 
 DECLARE_GPU_STAT_NAMED(ExrImgMediaReaderGpu, TEXT("ExrImgGpu"));
 DECLARE_GPU_STAT_NAMED(ExrImgMediaReaderGpu_MipRender, TEXT("ExrImgGpu.MipRender"));
@@ -668,14 +666,18 @@ FStructuredBufferPoolItemSharedPtr FExrImgMediaReaderGpu::AllocateGpuBufferFromP
 
 	if (!AllocatedBuffer)
 	{
-		// This boolean value is used to wait until render thread finishes buffer initialization.
-		volatile bool bInitDone = false;
+		FEvent* InitDoneEvent = nullptr;
+		if (bWait)
+		{
+			InitDoneEvent = FPlatformProcess::GetSynchEventFromPool();
+		}
+
 		{
 			AllocatedBuffer = MakeShareable(new FStructuredBufferPoolItem(), MoveTemp(BufferDeleter));
 			AllocatedBuffer->Reader = AsShared();
 
 			// Allocate and unlock the structured buffer on render thread.
-			ENQUEUE_RENDER_COMMAND(CreatePooledBuffer)([AllocatedBuffer, AllocSize, &bInitDone, this, bWait](FRHICommandListImmediate& RHICmdList)
+			ENQUEUE_RENDER_COMMAND(CreatePooledBuffer)([AllocatedBuffer, AllocSize, InitDoneEvent, this](FRHICommandListImmediate& RHICmdList)
 			{
 				TRACE_CPUPROFILER_EVENT_SCOPE_TEXT(*FString::Printf(TEXT("ExrReaderGpu.AllocBuffer_RenderThread")));
 				TRACE_CPUPROFILER_EVENT_SCOPE_TEXT(*FString::Printf(TEXT("ExrReaderGpu.AllocBuffer_RenderThread %d"), AllocSize));
@@ -696,19 +698,19 @@ FStructuredBufferPoolItemSharedPtr FExrImgMediaReaderGpu::AllocateGpuBufferFromP
 					AllocatedBuffer->ShaderResourceView = RHICreateShaderResourceView(AllocatedBuffer->UploadBufferRef);
 				}
 
-				if (bWait)
+				if (InitDoneEvent != nullptr)
 				{
-					bInitDone = true;
+					InitDoneEvent->Trigger();
 				}
 			});
 		}
 
 		/** Wait until buffer is initialized. */
-		while (!bInitDone && bWait)
+		if (InitDoneEvent != nullptr)
 		{
-			FPlatformProcess::Sleep(0.01f);
+			InitDoneEvent->Wait();
+			FPlatformProcess::ReturnSynchEventToPool(InitDoneEvent);
 		}
-
 	}
 
 	// This buffer will be automatically processed and returned to StagingMemoryPool once nothing keeps reference to it.
