ImgMedia: Use BeginRenderViewFamily now that they're correctly supported in scene captures (CL 25714191).

#rnx

diff --git a/Engine/Plugins/Media/ImgMedia/Source/ImgMedia/Private/ImgMediaSceneViewExtension.cpp b/Engine/Plugins/Media/ImgMedia/Source/ImgMedia/Private/ImgMediaSceneViewExtension.cpp
--- a/Engine/Plugins/Media/ImgMedia/Source/ImgMedia/Private/ImgMediaSceneViewExtension.cpp
+++ b/Engine/Plugins/Media/ImgMedia/Source/ImgMedia/Private/ImgMediaSceneViewExtension.cpp
@@ -41,32 +41,14 @@ FImgMediaSceneViewExtension::FImgMediaSceneViewExtension(const FAutoRegister& Au
 {
 }
 
-void FImgMediaSceneViewExtension::SetupViewFamily(FSceneViewFamily& InViewFamily)
-{
-}
-
-void FImgMediaSceneViewExtension::SetupView(FSceneViewFamily& InViewFamily, FSceneView& InView)
+void FImgMediaSceneViewExtension::BeginRenderViewFamily(FSceneViewFamily& InViewFamily)
 {
-	TRACE_CPUPROFILER_EVENT_SCOPE(FImgMediaSceneViewExtension::SetupView);
+	TRACE_CPUPROFILER_EVENT_SCOPE(FImgMediaSceneViewExtension::BeginRenderViewFamily);
 
 	/**
-	* NOTE: Scene captures call `SetupView` after the primary `BeginRenderViewFamily` call:
-	* FRendererModule::BeginRenderingViewFamilies
-	*     -> USceneCaptureComponent::UpdateDeferredCaptures
-	*         -> FScene::UpdateSceneCaptureContents
-	*             -> ISceneViewExtension::SetupView
+	* NOTE: Scene captures call this function after the primary `BeginRenderViewFamily` with the same frame number.
 	* Therefore, the view infos we cache here will be correctly kept until the next frame.
 	*/
-	if (InView.bIsSceneCapture)
-	{
-		CacheViewInfo(InViewFamily, InView);
-	}
-}
-
-void FImgMediaSceneViewExtension::BeginRenderViewFamily(FSceneViewFamily& InViewFamily)
-{
-	TRACE_CPUPROFILER_EVENT_SCOPE(FImgMediaSceneViewExtension::BeginRenderViewFamily);
-
 	if (LastFrameNumber != InViewFamily.FrameNumber)
 	{
 		CachedViewInfos.Reset();
@@ -75,8 +57,7 @@ void FImgMediaSceneViewExtension::BeginRenderViewFamily(FSceneViewFamily& InView
 
 	for (const FSceneView* View : InViewFamily.Views)
 	{
-		/** NOTE: Scene captures currently don't call BeginRenderViewFamily() so we handle them separately in SetupView. */
-		if (View && !View->bIsSceneCapture)
+		if (View != nullptr)
 		{
 			CacheViewInfo(InViewFamily, *View);
 		}
diff --git a/Engine/Plugins/Media/ImgMedia/Source/ImgMedia/Public/ImgMediaSceneViewExtension.h b/Engine/Plugins/Media/ImgMedia/Source/ImgMedia/Public/ImgMediaSceneViewExtension.h
--- a/Engine/Plugins/Media/ImgMedia/Source/ImgMedia/Public/ImgMediaSceneViewExtension.h
+++ b/Engine/Plugins/Media/ImgMedia/Source/ImgMedia/Public/ImgMediaSceneViewExtension.h
@@ -45,8 +45,8 @@ public:
 	 */
 	IMGMEDIA_API const TArray<FImgMediaViewInfo>& GetViewInfos() const { return CachedViewInfos; };
 
-	void SetupViewFamily(FSceneViewFamily& InViewFamily) override;
-	void SetupView(FSceneViewFamily& InViewFamily, FSceneView& InView) override;
+	void SetupViewFamily(FSceneViewFamily& InViewFamily) override {}
+	void SetupView(FSceneViewFamily& InViewFamily, FSceneView& InView) override {}
 	void BeginRenderViewFamily(FSceneViewFamily& InViewFamily) override;
 	int32 GetPriority() const override;
 
