ImgMedia: Improved buffer names to help tracking.

#rnx

diff --git a/Engine/Plugins/Media/ImgMedia/Source/ImgMedia/Private/Readers/ExrImgMediaReaderGpu.cpp b/Engine/Plugins/Media/ImgMedia/Source/ImgMedia/Private/Readers/ExrImgMediaReaderGpu.cpp
--- a/Engine/Plugins/Media/ImgMedia/Source/ImgMedia/Private/Readers/ExrImgMediaReaderGpu.cpp
+++ b/Engine/Plugins/Media/ImgMedia/Source/ImgMedia/Private/Readers/ExrImgMediaReaderGpu.cpp
@@ -684,12 +684,15 @@ FStructuredBufferPoolItemSharedPtr FExrImgMediaReaderGpu::AllocateGpuBufferFromP
 				FScopeLock ScopeLock(&AllocatorCriticalSecion);
 				SCOPED_GPU_STAT(RHICmdList, ExrImgMediaReaderGpu_AllocateBuffer);
 				SCOPED_DRAW_EVENT(RHICmdList, FExrImgMediaReaderGpu_AllocateBuffer);
-				FRHIResourceCreateInfo CreateInfo(TEXT("FExrImgMediaReaderGpu"));
+
+				FRHIResourceCreateInfo CreateInfo(TEXT(""));
+				CreateInfo.DebugName = TEXT("ExrReaderGpu.UploadBuffer");
 				AllocatedBuffer->UploadBufferRef = RHICmdList.CreateStructuredBuffer(sizeof(uint16) * 2., AllocSize, BUF_ShaderResource | BUF_Dynamic | BUF_FastVRAM, CreateInfo);
 				AllocatedBuffer->UploadBufferMapped = RHICmdList.LockBuffer(AllocatedBuffer->UploadBufferRef, 0, AllocSize, RLM_WriteOnly);
 
 				if (CVarExrReaderUseUploadHeap.GetValueOnAnyThread())
 				{
+					CreateInfo.DebugName = TEXT("ExrReaderGpu.DestBuffer");
 					AllocatedBuffer->ShaderAccessBufferRef = RHICmdList.CreateStructuredBuffer(sizeof(uint16) * 2., AllocSize, BUF_ShaderResource | BUF_FastVRAM, CreateInfo);
 					AllocatedBuffer->ShaderResourceView = RHICmdList.CreateShaderResourceView(AllocatedBuffer->ShaderAccessBufferRef);
 				}
