ImgMedia: Prevent zero requested frames on global cache.

#rnx

diff --git a/Engine/Plugins/Media/ImgMedia/Source/ImgMedia/Private/Loader/ImgMediaLoader.cpp b/Engine/Plugins/Media/ImgMedia/Source/ImgMedia/Private/Loader/ImgMediaLoader.cpp
--- a/Engine/Plugins/Media/ImgMedia/Source/ImgMedia/Private/Loader/ImgMediaLoader.cpp
+++ b/Engine/Plugins/Media/ImgMedia/Source/ImgMedia/Private/Loader/ImgMediaLoader.cpp
@@ -1101,8 +1101,8 @@ bool FImgMediaLoader::LoadSequence(const FString& SequencePath, const FFrameRate
 		NumFramesToLoad = FMath::Clamp(MaxFramesToLoad, MinFrames, NumImages);
 		const float LoadBehindScale = FMath::Clamp(Settings->CacheBehindPercentage, 0.0f, 100.0f) / 100.0f;
 
-		NumLoadBehind = (int32)(LoadBehindScale * MaxFramesToLoad);
-		NumLoadAhead = (int32)((1.0f - LoadBehindScale) * MaxFramesToLoad);
+		NumLoadBehind = (int32)(LoadBehindScale * NumFramesToLoad);
+		NumLoadAhead = FMath::CeilToInt32((1.0f - LoadBehindScale) * NumFramesToLoad);
 	}
 
 	Frames.Empty(NumFramesToLoad);
