ImgMedia: Fixed an issue with EXR calculating partial tiles incorrectly when total number of tiles is 1.

diff --git a/Engine/Plugins/Media/ImgMedia/Source/ExrReaderGpu/Private/ExrReaderGpu.cpp b/Engine/Plugins/Media/ImgMedia/Source/ExrReaderGpu/Private/ExrReaderGpu.cpp
--- a/Engine/Plugins/Media/ImgMedia/Source/ExrReaderGpu/Private/ExrReaderGpu.cpp
+++ b/Engine/Plugins/Media/ImgMedia/Source/ExrReaderGpu/Private/ExrReaderGpu.cpp
@@ -343,8 +343,7 @@ void FExrReader::CalculateTileOffsets
 				if (DimensionInTiles_PartialTiles.X != DimensionInTiles_CompleteTiles.X)
 				{
 					// If true - this is a partial tile in X dimension.
-					if (TileIndex != 0 &&
-						((TileIndex + 1) % DimensionInTiles_PartialTiles.X) == 0)
+					if (((TileIndex + 1) % DimensionInTiles_PartialTiles.X) == 0)
 					{
 						TileDim.X = PartialTileResolution.X;
 					}
diff --git a/Engine/Plugins/Media/ImgMedia/Source/ImgMedia/Private/Readers/ExrImgMediaReaderGpu.cpp b/Engine/Plugins/Media/ImgMedia/Source/ImgMedia/Private/Readers/ExrImgMediaReaderGpu.cpp
--- a/Engine/Plugins/Media/ImgMedia/Source/ImgMedia/Private/Readers/ExrImgMediaReaderGpu.cpp
+++ b/Engine/Plugins/Media/ImgMedia/Source/ImgMedia/Private/Readers/ExrImgMediaReaderGpu.cpp
@@ -268,7 +268,6 @@ bool FExrImgMediaReaderGpu::ReadFrame(int32 FrameId, const TMap<int32, FImgMedia
 
 	{
 		// Force mip level to be upscaled to all higher quality mips.
-		FIntPoint CurrentMipDim = ConverterParams->FullResolution;
 		TMap<int32, FImgMediaTileSelection> InMipTilesCopy = InMipTiles;
 		const int32 MipToUpscale = FMath::Clamp(Loader->GetMinimumLevelToUpscale(), -1, ConverterParams->NumMipLevels - 1);
 
@@ -359,7 +358,8 @@ bool FExrImgMediaReaderGpu::ReadFrame(int32 FrameId, const TMap<int32, FImgMedia
 			}
 
 			const int32 MipLevelDiv = 1 << CurrentMipLevel;
-			
+			FIntPoint CurrentMipDim = ConverterParams->FullResolution / MipLevelDiv;
+
 			TArray<FIntRect>& Viewports = ConverterParams->Viewports.Add(CurrentMipLevel);
 			for (const FIntRect& TileRegion : CurrentTileSelection.GetVisibleRegions())
 			{
