ImgMedia: Fix crash when image sequence has inconsistent resolution.

diff --git a/Engine/Plugins/Media/ImgMedia/Source/ImgMedia/Private/Readers/GenericImgMediaReader.cpp b/Engine/Plugins/Media/ImgMedia/Source/ImgMedia/Private/Readers/GenericImgMediaReader.cpp
--- a/Engine/Plugins/Media/ImgMedia/Source/ImgMedia/Private/Readers/GenericImgMediaReader.cpp
+++ b/Engine/Plugins/Media/ImgMedia/Source/ImgMedia/Private/Readers/GenericImgMediaReader.cpp
@@ -10,7 +10,7 @@
 #include "Misc/Paths.h"
 
 
-TSharedPtr<IImageWrapper> FGenericImgMediaReader::LoadFrameImage(const FString& ImagePath, TArray64<uint8>& OutBuffer, FImgMediaFrameInfo& OutInfo, bool bUseLoaderFirstFrame)
+TSharedPtr<IImageWrapper> FGenericImgMediaReader::LoadFrameImage(const FString& ImagePath, TArray64<uint8>& OutBuffer, FImgMediaFrameInfo& OutInfo)
 {
 	// load image into buffer
 	if (!FFileHelper::LoadFileToArray(OutBuffer, *ImagePath))
@@ -54,31 +54,13 @@ TSharedPtr<IImageWrapper> FGenericImgMediaReader::LoadFrameImage(const FString&
 		return nullptr;
 	}
 
-	if (bUseLoaderFirstFrame && LoaderPtr.IsValid())
-	{
-		const TSharedPtr<FImgMediaLoader, ESPMode::ThreadSafe> Loader = LoaderPtr.Pin();
-
-		OutInfo.Dim = Loader->GetSequenceDim();
-		OutInfo.NumMipLevels = Loader->GetNumMipLevels();
-
-		const SIZE_T SizeMip0 = (SIZE_T)OutInfo.Dim.X * OutInfo.Dim.Y * 4;
-		OutInfo.UncompressedSize = SizeMip0;
-		for (int32 Level = 1; Level < OutInfo.NumMipLevels; Level++)
-		{
-			OutInfo.UncompressedSize += SizeMip0 >> (2 * Level);
-		}
-	}
-	else
-	{
-		OutInfo.Dim.X = ImageWrapper->GetWidth();
-		OutInfo.Dim.Y = ImageWrapper->GetHeight();
-		OutInfo.UncompressedSize = (SIZE_T)OutInfo.Dim.X * OutInfo.Dim.Y * 4;
-		OutInfo.NumMipLevels = 1;
-	}
-
 	// get file info
 	const UImgMediaSettings* Settings = GetDefault<UImgMediaSettings>();
 	OutInfo.CompressionName = TEXT("");
+	OutInfo.Dim.X = ImageWrapper->GetWidth();
+	OutInfo.Dim.Y = ImageWrapper->GetHeight();
+	OutInfo.UncompressedSize = (SIZE_T)OutInfo.Dim.X * OutInfo.Dim.Y * 4;
+	OutInfo.NumMipLevels = 1;
 	OutInfo.FrameRate = Settings->DefaultFrameRate;
 	OutInfo.Srgb = true;
 	OutInfo.NumChannels = 4;
@@ -105,7 +87,7 @@ FGenericImgMediaReader::FGenericImgMediaReader(IImageWrapperModule& InImageWrapp
 bool FGenericImgMediaReader::GetFrameInfo(const FString& ImagePath, FImgMediaFrameInfo& OutInfo)
 {
 	TArray64<uint8> InputBuffer;
-	TSharedPtr<IImageWrapper> ImageWrapper = LoadFrameImage(ImagePath, InputBuffer, OutInfo, false);
+	TSharedPtr<IImageWrapper> ImageWrapper = LoadFrameImage(ImagePath, InputBuffer, OutInfo);
 
 	return ImageWrapper.IsValid();
 }
@@ -145,7 +127,7 @@ bool FGenericImgMediaReader::ReadFrame(int32 FrameId, const TMap<int32, FImgMedi
 
 				TArray64<uint8> InputBuffer;
 				FImgMediaFrameInfo Info;
-				TSharedPtr<IImageWrapper> ImageWrapper = LoadFrameImage(ImagePath, InputBuffer, Info, true);
+				TSharedPtr<IImageWrapper> ImageWrapper = LoadFrameImage(ImagePath, InputBuffer, Info);
 
 				if (!ImageWrapper.IsValid())
 				{
@@ -163,12 +145,27 @@ bool FGenericImgMediaReader::ReadFrame(int32 FrameId, const TMap<int32, FImgMedi
 				// Create the full buffer for data.
 				if (!OutFrame->Data.IsValid())
 				{
-					int64 AllocSize = (SIZE_T)BaseLevelDim.X * BaseLevelDim.Y * 4;
-					// Need more space for mips.
+					SIZE_T AllocSize = (SIZE_T)Info.Dim.X * Info.Dim.Y * 4;
+					
 					if (NumMipLevels > 1)
 					{
-						AllocSize = (AllocSize * 4) / 3;
+						ensureMsgf(Info.Dim.X <= BaseLevelDim.X && Info.Dim.Y <= BaseLevelDim.Y, TEXT("Invalid mip dimensions, larger than sequence dimensions."));
+
+						// We need to rely on the loader dimensions to support mips in separate folders.
+						Info.Dim = BaseLevelDim;
+						Info.NumMipLevels = NumMipLevels;
+
+						const SIZE_T SizeMip0 = (SIZE_T)BaseLevelDim.X * BaseLevelDim.Y * 4;
+						Info.UncompressedSize = SizeMip0;
+						for (int32 Level = 1; Level < Info.NumMipLevels; Level++)
+						{
+							Info.UncompressedSize += SizeMip0 >> (2 * Level);
+						}
+
+						// Need more space for mips.
+						AllocSize = (SizeMip0 * 4) / 3;
 					}
+
 					void* Buffer = FMemory::Malloc(AllocSize, PLATFORM_CACHE_LINE_SIZE);
 					OutFrame->Info = Info;
 					OutFrame->Data = MakeShareable(Buffer, [](void* ObjectToDelete) { FMemory::Free(ObjectToDelete); });
diff --git a/Engine/Plugins/Media/ImgMedia/Source/ImgMedia/Private/Readers/GenericImgMediaReader.h b/Engine/Plugins/Media/ImgMedia/Source/ImgMedia/Private/Readers/GenericImgMediaReader.h
--- a/Engine/Plugins/Media/ImgMedia/Source/ImgMedia/Private/Readers/GenericImgMediaReader.h
+++ b/Engine/Plugins/Media/ImgMedia/Source/ImgMedia/Private/Readers/GenericImgMediaReader.h
@@ -40,7 +40,7 @@ public:
 
 private:
 
-	TSharedPtr<IImageWrapper> LoadFrameImage(const FString& ImagePath, TArray64<uint8>& OutBuffer, FImgMediaFrameInfo& OutInfo, bool bUseLoaderFirstFrame);
+	TSharedPtr<IImageWrapper> LoadFrameImage(const FString& ImagePath, TArray64<uint8>& OutBuffer, FImgMediaFrameInfo& OutInfo);
 
 	/** The image wrapper module. */
 	IImageWrapperModule& ImageWrapperModule;
