ImgMedia: Removed properties from process exr that no longer being used.

diff --git a/Engine/Plugins/Media/ImgMedia/Source/ImgMediaEditor/Private/Widgets/ImgMediaProcessEXROptions.h b/Engine/Plugins/Media/ImgMedia/Source/ImgMediaEditor/Private/Widgets/ImgMediaProcessEXROptions.h
--- a/Engine/Plugins/Media/ImgMedia/Source/ImgMediaEditor/Private/Widgets/ImgMediaProcessEXROptions.h
+++ b/Engine/Plugins/Media/ImgMedia/Source/ImgMediaEditor/Private/Widgets/ImgMediaProcessEXROptions.h
@@ -47,14 +47,6 @@ public:
 	UPROPERTY(EditAnywhere, Transient, Category = Processing)
 	bool bUsePlayer = false;
 
-	/** Number of tiles in the X direction. If 0, then there are no tiles. */
-	UPROPERTY(VisibleAnywhere, Transient, Category = Tiles)
-	int32 NumTilesX = 0;
-
-	/** Number of tiles in the Y direction. If 0, then there are no tiles. */
-	UPROPERTY(VisibleAnywhere, Transient, Category = Tiles)
-	int32 NumTilesY = 0;
-
 	/** Tint each mip level a different colour to help with debugging. */
 	UPROPERTY(EditAnywhere, Transient, Category = Debug)
 	bool bEnableMipLevelTint = false;
