ImgMedia: Minor cleanup/deletion of old V1 player code.

#rnx

diff --git a/Engine/Plugins/Media/ImgMedia/Source/ImgMedia/Internal/Loader/ImgMediaLoader.h b/Engine/Plugins/Media/ImgMedia/Source/ImgMedia/Internal/Loader/ImgMediaLoader.h
--- a/Engine/Plugins/Media/ImgMedia/Source/ImgMedia/Internal/Loader/ImgMediaLoader.h
+++ b/Engine/Plugins/Media/ImgMedia/Source/ImgMedia/Internal/Loader/ImgMediaLoader.h
@@ -136,6 +136,7 @@ public:
 	 * @param Time The time of the image frame to get (relative to the beginning of the sequence).
 	 * @return The frame, or nullptr if the frame wasn't available yet.
 	 */
+	UE_DEPRECATED(5.4, "This method will be removed.")
 	TSharedPtr<FImgMediaTextureSample, ESPMode::ThreadSafe> GetFrameSample(FTimespan Time);
 
 	/**
diff --git a/Engine/Plugins/Media/ImgMedia/Source/ImgMedia/Private/Player/ImgMediaPlayer.cpp b/Engine/Plugins/Media/ImgMedia/Source/ImgMedia/Private/Player/ImgMediaPlayer.cpp
--- a/Engine/Plugins/Media/ImgMedia/Source/ImgMedia/Private/Player/ImgMediaPlayer.cpp
+++ b/Engine/Plugins/Media/ImgMedia/Source/ImgMedia/Private/Player/ImgMediaPlayer.cpp
@@ -23,7 +23,6 @@
 #include "MediaTexture.h"
 #include "MediaTextureTracker.h"
 
-#define IMG_MEDIA_PLAYER_VERSION 2
 #define LOCTEXT_NAMESPACE "FImgMediaPlayer"
 
 
@@ -347,125 +346,22 @@ void FImgMediaPlayer::TickInput(FTimespan DeltaTime, FTimespan /*Timecode*/)
 		}
 	}
 
-#if IMG_MEDIA_PLAYER_VERSION == 1
-
-	if ((CurrentState != EMediaState::Playing) || (CurrentDuration == FTimespan::Zero()))
-	{
-		return; // nothing to play
-	}
-
-	// update clock
-	if (PlaybackRestarted)
-	{
-		PlaybackRestarted = false;
-	}
-	else
-	{
-		CurrentTime += DeltaTime * CurrentRate;
-	}
-
-	// The following is a hack to accommodate for frame time rounding errors. The problem is
-	// that frame delta times can be one tick more or less each frame, depending on how the
-	// frame time is rounded. This presents a problem when driving media playback from Sequencer,
-	// because even both Sequencer and Media Framework clocks are running at the same rate, they
-	// may not be in phase with regards to rounding. This can cause some frames to be skipped.
-	// FFrameTime support in Media Framework is required to fix this properly.
-
-	if (!DeltaTimeHackApplied)
-	{
-		CurrentTime += HackDeltaTimeOffset;
-		DeltaTimeHackApplied = true;
-	}
-
-	// handle looping
-	if ((CurrentTime >= CurrentDuration) || (CurrentTime < FTimespan::Zero()))
-	{
-		EventSink.ReceiveMediaEvent(EMediaEvent::PlaybackEndReached);
-
-		if (ShouldLoop)
-		{
-			CurrentTime %= CurrentDuration;
-
-			if (CurrentTime < FTimespan::Zero())
-			{
-				CurrentTime += CurrentDuration;
-			}
-		}
-		else
-		{
-			CurrentState = EMediaState::Stopped;
-			CurrentTime = FTimespan::Zero();
-			CurrentRate = 0.0f;
-			LastNonZeroRate = 0.0f;
-			DeltaTimeHackApplied = false;
-
-			EventSink.ReceiveMediaEvent(EMediaEvent::PlaybackSuspended);
-		}
-	}
-
-	UE_LOG(LogImgMedia, VeryVerbose, TEXT("Player %p: CurrentTime %s, Delta %s, CurrentRate %f"),
-		this,
-		*CurrentTime.ToString(TEXT("%h:%m:%s.%t")),
-		*DeltaTime.ToString(TEXT("%h:%m:%s.%t")),
-		CurrentRate
-	);
-
-	// update image loader
-	if (SelectedVideoTrack == 0)
-	{
-		Loader->RequestFrame(CurrentTime, CurrentRate, ShouldLoop);
-	}
-	RequestFrameHasRun = true;
-#else
 	// Tick the scheduler an extra time in addition to its hookup as media clock sink, so we also get it moving forward during blocked playback
 	Scheduler->TickInput(FTimespan::Zero(), FTimespan::MinValue());
-#endif // IMG_MEDIA_PLAYER_VERSION == 1
 }
 
 bool FImgMediaPlayer::FlushOnSeekStarted() const
 {
-#if IMG_MEDIA_PLAYER_VERSION == 1
-	return false;
-#else
-	// V2 player response will always be treated as 'true'
 	return true;
-#endif
 }
 
 bool FImgMediaPlayer::FlushOnSeekCompleted() const
 {
-#if IMG_MEDIA_PLAYER_VERSION == 1
-	return true;
-#else
-	// V2 player response will always be treated as 'false'
 	return false;
-#endif
 }
 
 void FImgMediaPlayer::ProcessVideoSamples()
 {
-#if IMG_MEDIA_PLAYER_VERSION == 1
-	// Did we already run this frame?
-	if (RequestFrameHasRun)
-	{
-		RequestFrameHasRun = false;
-	}
-	else
-	{
-		// We are blocked... run stuff here as it will not get run normally.
-		if (Loader.IsValid())
-		{
-			if (SelectedVideoTrack == 0)
-			{
-				Loader->RequestFrame(CurrentTime, CurrentRate, ShouldLoop);
-			}
-		}
-		if (Scheduler.IsValid())
-		{
-			Scheduler->TickFetch(FTimespan::Zero(), FTimespan::Zero());
-		}
-	}
-#endif // IMG_MEDIA_PLAYER_VERSION == 1
 }
 
 //-----------------------------------------------------------------------------
@@ -474,7 +370,6 @@ void FImgMediaPlayer::ProcessVideoSamples()
 */
 bool FImgMediaPlayer::GetPlayerFeatureFlag(EFeatureFlag flag) const
 {
-#if IMG_MEDIA_PLAYER_VERSION >= 2
 	switch (flag)
 	{
 	case EFeatureFlag::UsePlaybackTimingV2:
@@ -483,7 +378,7 @@ bool FImgMediaPlayer::GetPlayerFeatureFlag(EFeatureFlag flag) const
 	default:
 		break;
 	}
-#endif  // IMG_MEDIA_PLAYER_VERSION >= 2
+	
 	return IMediaPlayer::GetPlayerFeatureFlag(flag);
 }
 
@@ -640,23 +535,12 @@ bool FImgMediaPlayer::Seek(const FTimespan& Time)
 		}
 	}
 
-#if IMG_MEDIA_PLAYER_VERSION == 1
-	// more timing hacks for Sequencer
-	CurrentTime = Time + HackDeltaTimeOffset;
-	DeltaTimeHackApplied = true;
-
-	if (CurrentTime == CurrentDuration)
-	{
-		CurrentTime -= HackDeltaTimeOffset;
-	}
-#else
 	CurrentTime = Time;
 	CurrentSeekIndex += (CurrentRate >= 0.0f) ? 1 : -1;
 	if (Loader.IsValid())
 	{
 		Loader->Seek(FMediaTimeStamp(CurrentTime, FMediaTimeStamp::MakeSequenceIndex(CurrentSeekIndex, 0)), LastNonZeroRate < 0.0f);
 	}
-#endif // IMG_MEDIA_PLAYER_VERSION == 1
 
 	if (CurrentState == EMediaState::Paused)
 	{
@@ -702,13 +586,6 @@ bool FImgMediaPlayer::SetRate(float Rate)
 	{
 		if (CurrentState == EMediaState::Stopped)
 		{
-#if IMG_MEDIA_PLAYER_VERSION == 1
-			if (Rate < 0.0f)
-			{
-				CurrentTime = CurrentDuration - FTimespan(1);
-			}
-#endif // IMG_MEDIA_PLAYER_VERSION == 1
-
 			PlaybackRestarted = true;
 		}
 
@@ -763,38 +640,7 @@ void FImgMediaPlayer::SetBlockingPlaybackHint(bool FacadeWillUseBlockingPlayback
 
 bool FImgMediaPlayer::FetchVideo(TRange<FTimespan> TimeRange, TSharedPtr<IMediaTextureSample, ESPMode::ThreadSafe>& OutSample)
 {
-#if IMG_MEDIA_PLAYER_VERSION == 1
-	if ((CurrentState != EMediaState::Paused) && (CurrentState != EMediaState::Playing))
-	{
-		return false; // nothing to play
-	}
-
-	if (SelectedVideoTrack != 0)
-	{
-		return false; // no video track selected
-	}
-
-	TSharedPtr<IMediaTextureSample, ESPMode::ThreadSafe> Sample = Loader->GetFrameSample(CurrentTime);
-
-	if (!Sample.IsValid())
-	{
-		return false; // sample not loaded yet
-	}
-
-	const FTimespan SampleTime = Sample->GetTime().Time;
-
-	if (SampleTime == LastFetchTime)
-	{
-		return false; // sample already fetched
-	}
-
-	LastFetchTime = SampleTime;
-	OutSample = Sample;
-
-	return true;
-#else // IMG_MEDIA_PLAYER_VERSION == 1
 	return false;
-#endif // IMG_MEDIA_PLAYER_VERSION == 1
 }
 
 
