ImgMedia: Fix global cache issue where cached empty frames could break info obtention on the first frame, and thus playback.

Also unified cached frame obtention method.

diff --git a/Engine/Plugins/Media/ImgMedia/Source/ImgMedia/Internal/Loader/ImgMediaLoader.h b/Engine/Plugins/Media/ImgMedia/Source/ImgMedia/Internal/Loader/ImgMediaLoader.h
--- a/Engine/Plugins/Media/ImgMedia/Source/ImgMedia/Internal/Loader/ImgMediaLoader.h
+++ b/Engine/Plugins/Media/ImgMedia/Source/ImgMedia/Internal/Loader/ImgMediaLoader.h
@@ -542,6 +542,11 @@ protected:
 	 */
 	void UpdateBandwidthThrottling();
 
+	/**
+	 * Get frame from either the local or global cache, nullptr if not found.
+	 */
+	const TSharedPtr<FImgMediaFrame, ESPMode::ThreadSafe>* GetCachedFrame(int32 InFrameNumber);
+
 private:
 
 	/** Critical section for synchronizing access to Frames. */
diff --git a/Engine/Plugins/Media/ImgMedia/Source/ImgMedia/Private/Loader/ImgMediaLoader.cpp b/Engine/Plugins/Media/ImgMedia/Source/ImgMedia/Private/Loader/ImgMediaLoader.cpp
--- a/Engine/Plugins/Media/ImgMedia/Source/ImgMedia/Private/Loader/ImgMediaLoader.cpp
+++ b/Engine/Plugins/Media/ImgMedia/Source/ImgMedia/Private/Loader/ImgMediaLoader.cpp
@@ -284,17 +284,8 @@ TSharedPtr<FImgMediaTextureSample, ESPMode::ThreadSafe> FImgMediaLoader::GetFram
 
 	FScopeLock ScopeLock(&CriticalSection);
 
-	const TSharedPtr<FImgMediaFrame, ESPMode::ThreadSafe>* Frame;
-	if (UseGlobalCache)
-	{
-		Frame = GlobalCache->FindAndTouch(SequenceName, FrameIndex);
-	}
-	else
-	{
-		Frame = Frames.FindAndTouch(FrameIndex);
-	}
-
-
+	const TSharedPtr<FImgMediaFrame, ESPMode::ThreadSafe>* Frame = GetCachedFrame(FrameIndex);
+	
 	if (Frame == nullptr)
 	{
 		return nullptr;
@@ -587,7 +578,7 @@ IMediaSamples::EFetchBestSampleResult FImgMediaLoader::FetchBestVideoSampleForTi
 			FScopeLock Lock(&CriticalSection);
 
 			// Get a frame if we have one available right now...
-			Frame = UseGlobalCache ? GlobalCache->FindAndTouch(SequenceName, MaxIdx) : Frames.FindAndTouch(MaxIdx);
+			Frame = GetCachedFrame(MaxIdx);
 
 			// Got a potential frame?
 			if (Frame)
@@ -763,7 +754,7 @@ bool FImgMediaLoader::PeekVideoSampleTime(FMediaTimeStamp &TimeStamp, bool bIsLo
 
 		// If possible, fetch any existing frame data...
 		// (just to see if we have any)
-		const TSharedPtr<FImgMediaFrame, ESPMode::ThreadSafe>* Frame = ((uint32)Idx != INDEX_NONE) ? (UseGlobalCache ? GlobalCache->FindAndTouch(SequenceName, Idx) : Frames.FindAndTouch(Idx)) : nullptr;
+		const TSharedPtr<FImgMediaFrame, ESPMode::ThreadSafe>* Frame = ((uint32)Idx != INDEX_NONE) ? GetCachedFrame(Idx) : nullptr;
 
 		// Data is present?
 		if (Frame)
@@ -831,19 +822,8 @@ IQueuedWork* FImgMediaLoader::GetWork()
 	FImgMediaLoaderWork* Work = (WorkPool.Num() > 0) ? WorkPool.Pop() : new FImgMediaLoaderWork(AsShared(), Reader.ToSharedRef());
 	
 	// Get the existing frame so we can add the mip level to it.
-	const TSharedPtr<FImgMediaFrame, ESPMode::ThreadSafe>* ExistingFramePtr;
-	if (UseGlobalCache)
-	{
-		ExistingFramePtr = GlobalCache->FindAndTouch(SequenceName, FrameNumber);
-	}
-	else
-	{
-		ExistingFramePtr = Frames.FindAndTouch(FrameNumber);
-	}
 	TSharedPtr<FImgMediaFrame, ESPMode::ThreadSafe> ExistingFrame;
-	
-	
-	if (ExistingFramePtr != nullptr)
+	if (const TSharedPtr<FImgMediaFrame, ESPMode::ThreadSafe>* ExistingFramePtr = GetCachedFrame(FrameNumber))
 	{
 		ExistingFrame = *ExistingFramePtr;
 	}
@@ -990,16 +970,12 @@ bool FImgMediaLoader::LoadSequence(const FString& SequencePath, const FFrameRate
 
 	// fetch sequence attributes from first image
 	{
-		// Try and get frame from the global cache.
-		const TSharedPtr<FImgMediaFrame, ESPMode::ThreadSafe>* Frame = nullptr;
-		if (UseGlobalCache)
-		{
-			Frame = GlobalCache->FindAndTouch(SequenceName, 0);
-		}
+		const TSharedPtr<FImgMediaFrame, ESPMode::ThreadSafe>* CachedFrame = GetCachedFrame(0);
 
-		if (Frame)
+		// Read info if the first frame is cached and not empty
+		if (CachedFrame && !(*CachedFrame)->MipTilesPresent.IsEmpty())
 		{
-			FirstFrameInfo = Frame->Get()->Info;
+			FirstFrameInfo = CachedFrame->Get()->Info;
 		}
 		else if (!Reader->GetFrameInfo(ImagePaths[0][0], FirstFrameInfo))
 		{
@@ -1449,15 +1425,7 @@ void FImgMediaLoader::Update(int32 PlayHeadFrame, float PlayRate, bool Loop)
 	{
 		// Get frame from cache.
 		bool NeedFrame = false;
-		const TSharedPtr<FImgMediaFrame, ESPMode::ThreadSafe>* FramePtr;
-		if (UseGlobalCache)
-		{
-			FramePtr = GlobalCache->FindAndTouch(SequenceName, FrameNumber);
-		}
-		else
-		{
-			FramePtr = Frames.FindAndTouch(FrameNumber);
-		}
+		const TSharedPtr<FImgMediaFrame, ESPMode::ThreadSafe>* FramePtr = GetCachedFrame(FrameNumber);
 
 		// Did we get a frame?
 		if ((FramePtr == nullptr) || ((*FramePtr).IsValid() == false))
@@ -1709,4 +1677,15 @@ void FImgMediaLoader::UpdateBandwidthThrottling()
 	}
 }
 
+const TSharedPtr<FImgMediaFrame, ESPMode::ThreadSafe>* FImgMediaLoader::GetCachedFrame(int32 InFrameNumber)
+{
+	if (UseGlobalCache)
+	{
+		return GlobalCache->FindAndTouch(SequenceName, InFrameNumber);
+	}
+	else
+	{
+		return Frames.FindAndTouch(InFrameNumber);
+	}
+}
 
