ImgMedia: Improved thread safety for shared frames in the global cache.

If a sequence with mips & tiles is shared between multiple players, there can be many concurrent updates the same shared cached frame (because they request different tile regions), which was previously not thread-safe. This change keeps the current logic, but ensures the thread-safety of operations on those shared frames.

* All members of FImgMediaFrame are now thread-safe.
    * A new FImgMediaMipTiles thread-safe container is used for MipTilesPresent.
    * Sample converter interface is now thread-safe (converter params are copied).
    * FImgMediaFrameInfo is accessed (or set) through copies.
* Prevent access to the gpu EXR reader memory pool critical section post-destruction.

diff --git a/Engine/Plugins/Media/ImgMedia/Source/ImgMedia/Private/GlobalCache/ImgMediaGlobalCache.cpp b/Engine/Plugins/Media/ImgMedia/Source/ImgMedia/Private/GlobalCache/ImgMediaGlobalCache.cpp
--- a/Engine/Plugins/Media/ImgMedia/Source/ImgMedia/Private/GlobalCache/ImgMediaGlobalCache.cpp
+++ b/Engine/Plugins/Media/ImgMedia/Source/ImgMedia/Private/GlobalCache/ImgMediaGlobalCache.cpp
@@ -65,7 +65,7 @@ void FImgMediaGlobalCache::AddFrame(const FString& FileName, const FName& Sequen
 	FScopeLock Lock(&CriticalSection);
 
 	// Make sure we have enough space in the cache to add this new frame.
-	SIZE_T FrameSize = Frame->Info.UncompressedSize;
+	SIZE_T FrameSize = Frame->GetUncompressedSize();
 	if (HasMipMaps)
 	{
 		FrameSize = (FrameSize * 4) / 3;
diff --git a/Engine/Plugins/Media/ImgMedia/Source/ImgMedia/Private/Loader/ImgMediaLoader.cpp b/Engine/Plugins/Media/ImgMedia/Source/ImgMedia/Private/Loader/ImgMediaLoader.cpp
--- a/Engine/Plugins/Media/ImgMedia/Source/ImgMedia/Private/Loader/ImgMediaLoader.cpp
+++ b/Engine/Plugins/Media/ImgMedia/Source/ImgMedia/Private/Loader/ImgMediaLoader.cpp
@@ -54,33 +54,6 @@ namespace ImgMediaLoader
 			InOutSequenceDim = InNewDim;
 		}
 	}
-	
-	// Check if the existing tiles contain all of the requested ones. (Is existing a superset of requested?)
-	bool ContainsMipTiles(const TMap<int32, FImgMediaTileSelection>& ExistingTiles, const TMap<int32, FImgMediaTileSelection>& RequestedTiles)
-	{
-		for (auto Iter = RequestedTiles.CreateConstIterator(); Iter; ++Iter)
-		{
-			const int32 RequestedMipLevel = Iter.Key();
-			const FImgMediaTileSelection& RequestedSelection = Iter.Value();
-
-			if (const FImgMediaTileSelection* ExistingSelection = ExistingTiles.Find(RequestedMipLevel))
-			{
-				if (!ExistingSelection->Contains(RequestedSelection))
-				{
-					// Requested tile selection is not present.
-					return false;
-				}
-			}
-			else
-			{
-				// Requested mip level is not present.
-				return false;
-			}
-		}
-		
-		// Requested tiles already exist, or the request was empty.
-		return true;
-	}
 
 	// Check if [CurrentFrame] is contained in the [OffsetCount] number of frames after [OriginFrame], also taking into account looping and play direction.
 	bool IsCachedFrameInRange(int32 CurrentFrame, int32 OriginFrame, int32 OffsetCount, int32 TotalNumFrames, float PlayRate)
@@ -123,19 +96,20 @@ void FImgMediaLoaderBandwidth::Update(const TSharedPtr<FImgMediaFrame, ESPMode::
 {
 	static constexpr int32 READ_TIME_CACHE_MAX = 64;
 
+	const FImgMediaFrameInfo FrameInfo = Frame->GetInfo();
 	// Calculate the current uncompressed bandwidth
-	SIZE_T BytesLoaded = Frame->Info.UncompressedSize;
+	SIZE_T BytesLoaded = FrameInfo.UncompressedSize;
 	Current = (float)BytesLoaded / WorkTime;
 
-	if (Frame->Info.bHasTiles)
+	if (FrameInfo.bHasTiles)
 	{
 		int32 TotalNumTiles = 0;
 
-		for (int32 MipLevel = 0; MipLevel < Frame->Info.NumMipLevels; ++MipLevel)
+		for (int32 MipLevel = 0; MipLevel < FrameInfo.NumMipLevels; ++MipLevel)
 		{
 			const int32 MipLevelDiv = 1 << MipLevel;
-			int32 NumTilesX = FMath::Max(1, FMath::CeilToInt(float(Frame->Info.NumTiles.X) / MipLevelDiv));
-			int32 NumTilesY = FMath::Max(1, FMath::CeilToInt(float(Frame->Info.NumTiles.Y) / MipLevelDiv));
+			int32 NumTilesX = FMath::Max(1, FMath::CeilToInt(float(FrameInfo.NumTiles.X) / MipLevelDiv));
+			int32 NumTilesY = FMath::Max(1, FMath::CeilToInt(float(FrameInfo.NumTiles.Y) / MipLevelDiv));
 			TotalNumTiles += NumTilesX * NumTilesY;
 		}
 
@@ -296,7 +270,7 @@ TSharedPtr<FImgMediaTextureSample, ESPMode::ThreadSafe> FImgMediaLoader::GetFram
 
 	auto Sample = MakeShared<FImgMediaTextureSample, ESPMode::ThreadSafe>();
 
-	ImgMediaLoader::CheckAndUpdateImgDimensions(SequenceDim, Frame->Get()->Info.Dim);
+	ImgMediaLoader::CheckAndUpdateImgDimensions(SequenceDim, Frame->Get()->GetDim());
 
 	if (!Sample->Initialize(*Frame->Get(), SequenceDim, FMediaTimeStamp(FrameStartTime, 0), NextStartTime - FrameStartTime, GetNumMipLevels(), TilingDescription))
 	{
@@ -618,13 +592,13 @@ IMediaSamples::EFetchBestSampleResult FImgMediaLoader::FetchBestVideoSampleForTi
 				// Track state & setup sample for caller...
 				QueuedSampleFetch.LastFrameIndex = MaxIdx;
 				QueuedSampleFetch.LastTimeStamp = FMediaTimeStamp(SampleTime, FMediaTimeStamp::MakeSequenceIndex(FMediaTimeStamp::GetPrimaryIndex(TimeRange.GetLowerBoundValue().SequenceIndex), QueuedSampleFetch.LoopIndex));
-				QueuedSampleFetch.LastDuration = FTimespan::FromSeconds(Frame->Get()->Info.FrameRate.AsInterval());
+				QueuedSampleFetch.LastDuration = FTimespan::FromSeconds(Frame->Get()->GetFrameRate().AsInterval());
 
 				// We are clear to return it as new result... Make a sample & initialize it...
 				auto Sample = MakeShared<FImgMediaTextureSample, ESPMode::ThreadSafe>();
 
 
-				ImgMediaLoader::CheckAndUpdateImgDimensions(SequenceDim, Frame->Get()->Info.Dim);
+				ImgMediaLoader::CheckAndUpdateImgDimensions(SequenceDim, Frame->Get()->GetDim());
 
 				if (Sample->Initialize(*Frame->Get(), SequenceDim, QueuedSampleFetch.LastTimeStamp, QueuedSampleFetch.LastDuration, GetNumMipLevels(), TilingDescription))
 				{
@@ -975,7 +949,7 @@ bool FImgMediaLoader::LoadSequence(const FString& SequencePath, const FFrameRate
 		// Read info if the first frame is cached and not empty
 		if (CachedFrame && !(*CachedFrame)->MipTilesPresent.IsEmpty())
 		{
-			FirstFrameInfo = CachedFrame->Get()->Info;
+			FirstFrameInfo = CachedFrame->Get()->GetInfo();
 		}
 		else if (!Reader->GetFrameInfo(ImagePaths[0][0], FirstFrameInfo))
 		{
@@ -1439,7 +1413,7 @@ void FImgMediaLoader::Update(int32 PlayHeadFrame, float PlayRate, bool Loop)
 			TMap<int32, FImgMediaTileSelection> DesiredMipsAndTiles;
 			GetDesiredMipTiles(FrameNumber, DesiredMipsAndTiles);
 
-			NeedFrame = !ImgMediaLoader::ContainsMipTiles((*FramePtr)->MipTilesPresent, DesiredMipsAndTiles);
+			NeedFrame = !(*FramePtr)->MipTilesPresent.ContainsMipTiles(DesiredMipsAndTiles);
 
 			if (NeedFrame)
 			{
@@ -1498,12 +1472,16 @@ bool FImgMediaLoader::TryAddEmptyFrame(int32 FrameNumber)
 		const int32 NumChannels = 3;
 		const int32 PixelSize = sizeof(uint16) * NumChannels;
 		TSharedPtr<FImgMediaFrame, ESPMode::ThreadSafe> Frame = MakeShared<FImgMediaFrame>();
-		Frame->Info.Dim = SequenceDim;
-		Frame->Info.FrameRate = SequenceFrameRate;
-		Frame->Info.NumChannels = NumChannels;
-		Frame->Info.bHasTiles = false;
+		
+		FImgMediaFrameInfo FrameInfo = {};
+		FrameInfo.Dim = SequenceDim;
+		FrameInfo.FrameRate = SequenceFrameRate;
+		FrameInfo.NumChannels = NumChannels;
+		FrameInfo.bHasTiles = false;
+		Frame->SetInfo(FrameInfo);
+
 		Frame->Format = EMediaTextureSampleFormat::FloatRGB;
-		Frame->Stride = Frame->Info.Dim.X * PixelSize;
+		Frame->Stride = SequenceDim.X * PixelSize;
 
 		AddFrameToCache(FrameNumber, Frame);
 	}
diff --git a/Engine/Plugins/Media/ImgMedia/Source/ImgMedia/Private/Player/ImgMediaTextureSample.h b/Engine/Plugins/Media/ImgMedia/Source/ImgMedia/Private/Player/ImgMediaTextureSample.h
--- a/Engine/Plugins/Media/ImgMedia/Source/ImgMedia/Private/Player/ImgMediaTextureSample.h
+++ b/Engine/Plugins/Media/ImgMedia/Source/ImgMedia/Private/Player/ImgMediaTextureSample.h
@@ -67,7 +67,7 @@ public:
 
 	virtual FIntPoint GetDim() const override
 	{
-		return Frame.Info.Dim;
+		return Frame.GetDim();
 	}
 
 	virtual uint8 GetNumMips() const override
@@ -124,7 +124,7 @@ public:
 
 	virtual bool IsOutputSrgb() const override
 	{
-		return Frame.Info.Srgb;
+		return Frame.IsOutputSrgb();
 	}
 
 private:
diff --git a/Engine/Plugins/Media/ImgMedia/Source/ImgMedia/Private/Readers/ExrImgMediaReader.cpp b/Engine/Plugins/Media/ImgMedia/Source/ImgMedia/Private/Readers/ExrImgMediaReader.cpp
--- a/Engine/Plugins/Media/ImgMedia/Source/ImgMedia/Private/Readers/ExrImgMediaReader.cpp
+++ b/Engine/Plugins/Media/ImgMedia/Source/ImgMedia/Private/Readers/ExrImgMediaReader.cpp
@@ -60,7 +60,7 @@ FExrImgMediaReader::EReadResult FExrImgMediaReader::ReadTiles
 	, int64 BufferSize
 	, const FString& ImagePath
 	, const TArray<FIntRect>& TileRegions
-	, TSharedPtr<FSampleConverterParameters> ConverterParams
+	, FSampleConverterParameters& ConverterParams
 	, const int32 CurrentMipLevel
 	, TArray<UE::Math::TIntPoint<int64>>& OutBufferRegionsToCopy)
 {
@@ -70,29 +70,29 @@ FExrImgMediaReader::EReadResult FExrImgMediaReader::ReadTiles
 	FExrReader ChunkReader;
 	int MipLevelDiv = 1 << CurrentMipLevel;
 
-	FIntPoint MipResolution = ConverterParams->FullResolution / MipLevelDiv;
+	FIntPoint MipResolution = ConverterParams.FullResolution / MipLevelDiv;
 
 	FIntPoint DimensionInTiles
-		( FMath::CeilToInt(float(MipResolution.X) / ConverterParams->TileDimWithBorders.X)
-		, FMath::CeilToInt(float(MipResolution.Y) / ConverterParams->TileDimWithBorders.Y));
+		( FMath::CeilToInt(float(MipResolution.X) / ConverterParams.TileDimWithBorders.X)
+		, FMath::CeilToInt(float(MipResolution.Y) / ConverterParams.TileDimWithBorders.Y));
 
 	TArray<int32> NumTilesPerLevel;
 	TArray<TArray<int64>> CustomOffsets;
-	int32 NumMipLevels = ConverterParams->bMipsInSeparateFiles ? 1 : ConverterParams->NumMipLevels;
+	int32 NumMipLevels = ConverterParams.bMipsInSeparateFiles ? 1 : ConverterParams.NumMipLevels;
 
 	FExrReader::CalculateTileOffsets(
 		NumTilesPerLevel,
 		CustomOffsets,
-		ConverterParams->TileInfoPerMipLevel,
-		ConverterParams->FullResolution,
-		ConverterParams->TileDimWithBorders,
+		ConverterParams.TileInfoPerMipLevel,
+		ConverterParams.FullResolution,
+		ConverterParams.TileDimWithBorders,
 		NumMipLevels,
-		ConverterParams->PixelSize,
-		ConverterParams->bCustomExr);
+		ConverterParams.PixelSize,
+		ConverterParams.bCustomExr);
 
 	{
 		TRACE_CPUPROFILER_EVENT_SCOPE_TEXT(*FString::Printf(TEXT("FExrImgMediaReader_ReadTilesCustom_OpenFile")));
-		if (!ChunkReader.OpenExrAndPrepareForPixelReading(ImagePath, NumTilesPerLevel, MoveTemp(CustomOffsets), ConverterParams->bCustomExr))
+		if (!ChunkReader.OpenExrAndPrepareForPixelReading(ImagePath, NumTilesPerLevel, MoveTemp(CustomOffsets), ConverterParams.bCustomExr))
 		{
 			return Fail;
 		}
@@ -112,22 +112,22 @@ FExrImgMediaReader::EReadResult FExrImgMediaReader::ReadTiles
 				// Check to see if the frame was canceled.
 				{
 					FScopeLock RegionScopeLock(&CanceledFramesCriticalSection);
-					if (CanceledFrames.Remove(ConverterParams->FrameId) > 0)
+					if (CanceledFrames.Remove(ConverterParams.FrameId) > 0)
 					{
-						UE_LOG(LogImgMedia, Verbose, TEXT("Reader %p: Canceling Frame %i At tile row # %i"), this, ConverterParams->FrameId, TileRow);
+						UE_LOG(LogImgMedia, Verbose, TEXT("Reader %p: Canceling Frame %i At tile row # %i"), this, ConverterParams.FrameId, TileRow);
 						Result = Cancelled;
 						break;
 					}
 				}
 
-				const uint16 Padding = ConverterParams->bCustomExr ? 0 : FExrReader::TILE_PADDING;
-				const int64 TileByteStride = ConverterParams->PixelSize * ConverterParams->TileDimWithBorders.X * ConverterParams->TileDimWithBorders.Y + Padding;
+				const uint16 Padding = ConverterParams.bCustomExr ? 0 : FExrReader::TILE_PADDING;
+				const int64 TileByteStride = ConverterParams.PixelSize * ConverterParams.TileDimWithBorders.X * ConverterParams.TileDimWithBorders.Y + Padding;
 				const int StartTileIndex = TileRow * DimensionInTiles.X + TileRegion.Min.X;
 
 				bool bLastTile = TileRow + 1 == DimensionInTiles.Y && TileRegion.Max.X == DimensionInTiles.X;
 				const int EndTileIndex = TileRow * DimensionInTiles.X + TileRegion.Max.X;
 
-				int MipLevel = ConverterParams->bMipsInSeparateFiles ? 0 : CurrentMipLevel;
+				int MipLevel = ConverterParams.bMipsInSeparateFiles ? 0 : CurrentMipLevel;
 				ChunkReader.SeekTileWithinFile(StartTileIndex, MipLevel, CurrentBufferPos);
 				int64 ByteOffsetStart = 0;
 				int64 ByteOffsetEnd = 0;
@@ -185,18 +185,19 @@ bool FExrImgMediaReader::ReadFrame(int32 FrameId, const TMap<int32, FImgMediaTil
 	int32 BytesPerPixelPerChannel = sizeof(uint16);
 	int32 NumChannels = 4;
 	int32 BytesPerPixel = BytesPerPixelPerChannel * NumChannels;
+	FImgMediaFrameInfo FrameInfo;
 
 	// Do we already have our buffer?
 	if (OutFrame->Data.IsValid() == false)
 	{
 		// Nope. Create it.
 		const FString& LargestImage = Loader->GetImagePath(FrameId, 0);
-		if (!GetInfo(LargestImage, OutFrame->Info))
+		if (!GetInfo(LargestImage, FrameInfo, OutFrame))
 		{
 			return false;
 		}
 
-		const FIntPoint& Dim = OutFrame->Info.Dim;
+		const FIntPoint& Dim = FrameInfo.Dim;
 
 		if (Dim.GetMin() <= 0)
 		{
@@ -232,12 +233,16 @@ bool FExrImgMediaReader::ReadFrame(int32 FrameId, const TMap<int32, FImgMediaTil
 		OutFrame->Format = EMediaTextureSampleFormat::FloatRGBA;
 		OutFrame->Data = MakeShareable(Buffer, MoveTemp(BufferDeleter));
 		OutFrame->MipTilesPresent.Reset();
-		OutFrame->Stride = OutFrame->Info.Dim.X * BytesPerPixel;
+		OutFrame->Stride = FrameInfo.Dim.X * BytesPerPixel;
+	}
+	else
+	{
+		FrameInfo = OutFrame->GetInfo();
 	}
 
 	// Loop over all mips.
 	uint8* MipDataPtr = (uint8*)(OutFrame->Data.Get());
-	FIntPoint Dim = OutFrame->Info.Dim;
+	FIntPoint Dim = FrameInfo.Dim;
 	
 	int32 NumMipLevels = Loader->GetNumMipLevels();
 	for (int32 CurrentMipLevel = 0; CurrentMipLevel < NumMipLevels; ++CurrentMipLevel)
@@ -248,30 +253,25 @@ bool FExrImgMediaReader::ReadFrame(int32 FrameId, const TMap<int32, FImgMediaTil
 
 			const int MipLevelDiv = 1 << CurrentMipLevel;
 
-			bool ReadThisMip = true;
 			// Avoid reads if the cached frame already contains the current tiles for this mip level.
-			if (const FImgMediaTileSelection* CachedSelection = OutFrame->MipTilesPresent.Find(CurrentMipLevel))
-			{
-				ReadThisMip = !CachedSelection->Contains(CurrentTileSelection);
-			}
-
+			const bool ReadThisMip = !OutFrame->MipTilesPresent.ContainsTiles(CurrentMipLevel, CurrentTileSelection);
 			if (ReadThisMip)
 			{
 				FString Image = Loader->GetImagePath(FrameId, CurrentMipLevel);
 				FString BaseImage;
 
-				if (OutFrame->Info.FormatName == TEXT("EXR CUSTOM"))
+				if (FrameInfo.FormatName == TEXT("EXR CUSTOM"))
 				{
 #if defined(PLATFORM_WINDOWS) && PLATFORM_WINDOWS
 					TArray<FIntRect> TileRegions = CurrentTileSelection.GetVisibleRegions();
-					int32 PixelSize = sizeof(uint16) * OutFrame->Info.NumChannels;
-					TSharedPtr<FSampleConverterParameters> ConverterParams = MakeShared<FSampleConverterParameters>();
-					ConverterParams->FrameInfo = OutFrame->Info;
-					ConverterParams->PixelSize = sizeof(uint16) * ConverterParams->FrameInfo.NumChannels;
-					ConverterParams->TileDimWithBorders = OutFrame->Info.TileDimensions + OutFrame->Info.TileBorder * 2;
-					ConverterParams->NumMipLevels = Loader->GetNumMipLevels();
-					ConverterParams->bCustomExr = OutFrame->Info.FormatName == TEXT("EXR CUSTOM");
-					ConverterParams->FrameId = FrameId;
+					int32 PixelSize = sizeof(uint16) * FrameInfo.NumChannels;
+					FSampleConverterParameters ConverterParams;
+					ConverterParams.FrameInfo = FrameInfo;
+					ConverterParams.PixelSize = sizeof(uint16) * ConverterParams.FrameInfo.NumChannels;
+					ConverterParams.TileDimWithBorders = FrameInfo.TileDimensions + FrameInfo.TileBorder * 2;
+					ConverterParams.NumMipLevels = Loader->GetNumMipLevels();
+					ConverterParams.bCustomExr = FrameInfo.FormatName == TEXT("EXR CUSTOM");
+					ConverterParams.FrameId = FrameId;
 
 					TArray<UE::Math::TIntPoint<int64>> OutBufferRegionsToCopy;
 					EReadResult ReadResult = ReadTiles((uint16*)MipDataPtr, GetMipBufferTotalSize(Dim / MipLevelDiv), Image, TileRegions, ConverterParams, CurrentMipLevel, OutBufferRegionsToCopy);
@@ -382,7 +382,7 @@ TSharedPtr<IImgMediaReader, ESPMode::ThreadSafe> FExrImgMediaReader::GetReader(c
 /* FExrImgMediaReader implementation
  *****************************************************************************/
 
-bool FExrImgMediaReader::GetInfo(const FString& FilePath, FImgMediaFrameInfo& OutInfo)
+bool FExrImgMediaReader::GetInfo(const FString& FilePath, FImgMediaFrameInfo& OutInfo, const TSharedPtr<FImgMediaFrame, ESPMode::ThreadSafe>& OutFrame)
 {
 	FOpenExrHeaderReader HeaderReader(FilePath);
 	if (HeaderReader.HasInputFile() == false)
@@ -440,6 +440,10 @@ bool FExrImgMediaReader::GetInfo(const FString& FilePath, FImgMediaFrameInfo& Ou
 		OutInfo.NumTiles = FIntPoint(1, 1);
 	}
 
+	if (OutFrame.IsValid())
+	{
+		OutFrame->SetInfo(OutInfo);
+	}
 
 	return (OutInfo.UncompressedSize > 0) && (OutInfo.Dim.GetMin() > 0);
 }
diff --git a/Engine/Plugins/Media/ImgMedia/Source/ImgMedia/Private/Readers/ExrImgMediaReader.h b/Engine/Plugins/Media/ImgMedia/Source/ImgMedia/Private/Readers/ExrImgMediaReader.h
--- a/Engine/Plugins/Media/ImgMedia/Source/ImgMedia/Private/Readers/ExrImgMediaReader.h
+++ b/Engine/Plugins/Media/ImgMedia/Source/ImgMedia/Private/Readers/ExrImgMediaReader.h
@@ -57,9 +57,10 @@ protected:
 	 *
 	 * @param FilePath The location of the exr file.
 	 * @param OutInfo Will contain the frame information.
+	 * @param OutFrame Optional frame to update.
 	 * @return true on success, false otherwise.
 	 */
-	static bool GetInfo(const FString& FilePath, FImgMediaFrameInfo& OutInfo);
+	static bool GetInfo(const FString& FilePath, FImgMediaFrameInfo& OutInfo, const TSharedPtr<FImgMediaFrame, ESPMode::ThreadSafe>& OutFrame = nullptr);
 
 	/**
 	 * Reads tiles from exr files in tile rows based on Tile region. If frame is pending for cancelation
@@ -78,7 +79,7 @@ protected:
 		, int64 BufferSize
 		, const FString& ImagePath
 		, const TArray<FIntRect>& TileRegions
-		, TSharedPtr<FSampleConverterParameters> ConverterParams
+		, FSampleConverterParameters& ConverterParams
 		, const int32 CurrentMipLevel
 		, TArray<UE::Math::TIntPoint<int64>>& OutBufferRegionsToCopy);
 
diff --git a/Engine/Plugins/Media/ImgMedia/Source/ImgMedia/Private/Readers/ExrImgMediaReaderGpu.cpp b/Engine/Plugins/Media/ImgMedia/Source/ImgMedia/Private/Readers/ExrImgMediaReaderGpu.cpp
--- a/Engine/Plugins/Media/ImgMedia/Source/ImgMedia/Private/Readers/ExrImgMediaReaderGpu.cpp
+++ b/Engine/Plugins/Media/ImgMedia/Source/ImgMedia/Private/Readers/ExrImgMediaReaderGpu.cpp
@@ -96,7 +96,7 @@ FExrImgMediaReaderGpu::FExrImgMediaReaderGpu(const TSharedRef<FImgMediaLoader, E
 
 FExrImgMediaReaderGpu::~FExrImgMediaReaderGpu()
 {
-	FScopeLock ScopeLock(&AllocatorCriticalSecion);
+	FScopeLock ScopeLock(&MemoryPoolCriticalSection);
 
 	// Copy memory pool array to be released on render thread.
 	ENQUEUE_RENDER_COMMAND(DeletePooledBuffers)([InMemoryPool = MemoryPool](FRHICommandListImmediate& RHICmdList)
@@ -120,28 +120,26 @@ FExrImgMediaReader::EReadResult FExrImgMediaReaderGpu::ReadMip
 	( const int32 CurrentMipLevel
 	, const FImgMediaTileSelection& CurrentTileSelection
 	, TSharedPtr<FImgMediaFrame, ESPMode::ThreadSafe> OutFrame
+	, FSampleConverterParameters& ConverterParams
 	, TSharedPtr<FExrMediaTextureSampleConverter, ESPMode::ThreadSafe> SampleConverter
-	, const FString& ImagePath
-	, bool bHasTiles)
+	, const FString& ImagePath)
 {
 
 	TRACE_CPUPROFILER_EVENT_SCOPE_TEXT(*FString::Printf(TEXT("ExrReaderGpu.ReadMip %d"), CurrentMipLevel));
 
 	// Next mip level.
 	int MipLevelDiv = 1 << CurrentMipLevel;
-	FIntPoint CurrentMipDim = SampleConverter->ConverterParams->FullResolution / MipLevelDiv;
-
-	const SIZE_T BufferSize = GetBufferSize(CurrentMipDim, SampleConverter->ConverterParams->FrameInfo.NumChannels, bHasTiles, OutFrame->Info.NumTiles / MipLevelDiv, SampleConverter->ConverterParams->bCustomExr);
+	FIntPoint CurrentMipDim = ConverterParams.FullResolution / MipLevelDiv;
+	const FImgMediaFrameInfo& FrameInfo = ConverterParams.FrameInfo;
+	const SIZE_T BufferSize = GetBufferSize(CurrentMipDim, FrameInfo.NumChannels, FrameInfo.bHasTiles, FrameInfo.NumTiles / MipLevelDiv, ConverterParams.bCustomExr);
 	
-	SampleConverter->LockMipBuffers();
-	FStructuredBufferPoolItemSharedPtr BufferData = SampleConverter->GetMipLevelBuffer(CurrentMipLevel);
-
-	if (!BufferData.IsValid())
-	{
-		BufferData = AllocateGpuBufferFromPool(BufferSize);
-		SampleConverter->SetMipLevelBuffer(CurrentMipLevel, BufferData);
-	}
-	SampleConverter->UnlockMipBuffers();
+	FStructuredBufferPoolItemSharedPtr BufferData = SampleConverter->GetOrCreateMipLevelBuffer(
+		CurrentMipLevel,
+		[this, BufferSize]() -> FStructuredBufferPoolItemSharedPtr
+		{
+			return AllocateGpuBufferFromPool(BufferSize);
+		}
+	);
 
 	uint16* MipDataPtr = static_cast<uint16*>(BufferData->UploadBufferMapped);
 
@@ -151,17 +149,13 @@ FExrImgMediaReader::EReadResult FExrImgMediaReaderGpu::ReadMip
 	{
 		TArray<UE::Math::TIntPoint<int64>> BufferRegionsToCopy;
 		// read frame data
-		if (bHasTiles || SampleConverter->ConverterParams->bCustomExr)
+		if (FrameInfo.bHasTiles || ConverterParams.bCustomExr)
 		{
 			TArray<FIntRect> TileRegionsToRead;
 			{
 				TRACE_CPUPROFILER_EVENT_SCOPE_TEXT(*FString::Printf(TEXT("ExrReaderGpu.CalculateRegions %d"), CurrentMipLevel));
 
-				if (const FImgMediaTileSelection* CachedSelection = OutFrame->MipTilesPresent.Find(CurrentMipLevel))
-				{
-					TileRegionsToRead = CachedSelection->GetVisibleRegions(&CurrentTileSelection);
-				}
-				else
+				if (!OutFrame->MipTilesPresent.GetVisibleRegions(CurrentMipLevel, CurrentTileSelection, TileRegionsToRead))
 				{
 					TileRegionsToRead = CurrentTileSelection.GetVisibleRegions();
 				}
@@ -174,7 +168,7 @@ FExrImgMediaReader::EReadResult FExrImgMediaReaderGpu::ReadMip
 			}
 			else
 			{
-				ReadResult = ReadTiles(MipDataPtr, BufferSize, ImagePath, TileRegionsToRead, SampleConverter->ConverterParams, CurrentMipLevel, BufferRegionsToCopy);
+				ReadResult = ReadTiles(MipDataPtr, BufferSize, ImagePath, TileRegionsToRead, ConverterParams, CurrentMipLevel, BufferRegionsToCopy);
 				
 				for (const FIntRect& Region : TileRegionsToRead)
 				{
@@ -184,16 +178,16 @@ FExrImgMediaReader::EReadResult FExrImgMediaReaderGpu::ReadMip
 		}
 		else
 		{
-			ReadResult = ReadInChunks(MipDataPtr, ImagePath, SampleConverter->ConverterParams->FrameId, CurrentMipDim, BufferSize);
+			ReadResult = ReadInChunks(MipDataPtr, ImagePath, ConverterParams.FrameId, CurrentMipDim, BufferSize);
 			OutFrame->NumTilesRead++;
 		}
 
 		if (ReadResult == Success && bUseUploadHeap)
 		{
-			ENQUEUE_RENDER_COMMAND(CopyFromUploadBuffer)([SampleConverter, BufferData, BufferRegionsToCopy](FRHICommandListImmediate& RHICmdList)
+			ENQUEUE_RENDER_COMMAND(CopyFromUploadBuffer)([SampleConverter, BufferData, BufferRegionsToCopy, FrameId = ConverterParams.FrameId](FRHICommandListImmediate& RHICmdList)
 				{
 					SCOPED_DRAW_EVENT(RHICmdList, FExrImgMediaReaderGpu_CopyBuffers);
-					TRACE_CPUPROFILER_EVENT_SCOPE_TEXT(*FString::Printf(TEXT("ExrReaderGpu.StartCopy %d"), SampleConverter->ConverterParams->FrameId));
+					TRACE_CPUPROFILER_EVENT_SCOPE_TEXT(*FString::Printf(TEXT("ExrReaderGpu.StartCopy %d"), FrameId));
 					SCOPED_GPU_STAT(RHICmdList, ExrImgMediaReaderGpu_CopyUploadBuffer);
 
 					if (BufferRegionsToCopy.IsEmpty())
@@ -235,11 +229,10 @@ bool FExrImgMediaReaderGpu::ReadFrame(int32 FrameId, const TMap<int32, FImgMedia
 	{
 		return false;
 	}
-
-
+	
 	const FString& LargestImagePath = Loader->GetImagePath(FrameId, 0);
-
-	if (!GetInfo(LargestImagePath, OutFrame->Info))
+	FImgMediaFrameInfo FrameInfo;
+	if (!GetInfo(LargestImagePath, FrameInfo, OutFrame))
 	{
 		return false;
 	}
@@ -247,35 +240,31 @@ bool FExrImgMediaReaderGpu::ReadFrame(int32 FrameId, const TMap<int32, FImgMedia
 	TSharedPtr<FExrMediaTextureSampleConverter, ESPMode::ThreadSafe> SampleConverterPtr = OutFrame->GetOrCreateSampleConverter<FExrMediaTextureSampleConverter>();
 
 	// Get tile info.
-	bool bHasTiles = OutFrame->Info.bHasTiles;
-	TSharedPtr<FSampleConverterParameters> ConverterParams = SampleConverterPtr->ConverterParams;
-	ConverterParams->FullResolution = OutFrame->Info.Dim;
-	ConverterParams->FrameId = FrameId;
-
-	const FIntPoint TileDim = OutFrame->Info.TileDimensions;
-
-	if (ConverterParams->FullResolution.GetMin() <= 0)
+	FSampleConverterParameters ConverterParams = SampleConverterPtr->GetParams();
+	ConverterParams.FullResolution = FrameInfo.Dim;
+	ConverterParams.FrameId = FrameId;
+	if (ConverterParams.FullResolution.GetMin() <= 0)
 	{
 		return false;
 	}
 
-	ConverterParams->FrameInfo = OutFrame->Info;
-	ConverterParams->PixelSize = sizeof(uint16) * ConverterParams->FrameInfo.NumChannels;
-	ConverterParams->TileDimWithBorders = TileDim + OutFrame->Info.TileBorder * 2;
-	ConverterParams->NumMipLevels = Loader->GetNumMipLevels();
-	ConverterParams->bCustomExr = OutFrame->Info.FormatName == TEXT("EXR CUSTOM");
-	ConverterParams->bMipsInSeparateFiles = Loader->MipsInSeparateFiles();
+	ConverterParams.FrameInfo = FrameInfo;
+	ConverterParams.PixelSize = sizeof(uint16) * ConverterParams.FrameInfo.NumChannels;
+	ConverterParams.TileDimWithBorders = FrameInfo.TileDimensions + FrameInfo.TileBorder * 2;
+	ConverterParams.NumMipLevels = Loader->GetNumMipLevels();
+	ConverterParams.bCustomExr = FrameInfo.FormatName == TEXT("EXR CUSTOM");
+	ConverterParams.bMipsInSeparateFiles = Loader->MipsInSeparateFiles();
 
 	{
 		// Force mip level to be upscaled to all higher quality mips.
 		TMap<int32, FImgMediaTileSelection> InMipTilesCopy = InMipTiles;
-		const int32 MipToUpscale = FMath::Clamp(Loader->GetMinimumLevelToUpscale(), -1, ConverterParams->NumMipLevels - 1);
+		const int32 MipToUpscale = FMath::Clamp(Loader->GetMinimumLevelToUpscale(), -1, ConverterParams.NumMipLevels - 1);
 
-		if (ConverterParams->NumMipLevels > 1 && MipToUpscale >= 0)
+		if (ConverterParams.NumMipLevels > 1 && MipToUpscale >= 0)
 		{
-			ConverterParams->UpscaleMip = MipToUpscale;
+			ConverterParams.UpscaleMip = MipToUpscale;
 
-			FImgMediaTileSelection FullSelection = FImgMediaTileSelection::CreateForTargetMipLevel(ConverterParams->FullResolution, TileDim, MipToUpscale, true);
+			FImgMediaTileSelection FullSelection = FImgMediaTileSelection::CreateForTargetMipLevel(ConverterParams.FullResolution, FrameInfo.TileDimensions, MipToUpscale, true);
 			if (InMipTilesCopy.Contains(MipToUpscale))
 			{
 				InMipTilesCopy[MipToUpscale] = FullSelection;
@@ -298,20 +287,13 @@ bool FExrImgMediaReaderGpu::ReadFrame(int32 FrameId, const TMap<int32, FImgMedia
 			}
 
 			// Get highest resolution mip level path.
-			FString ImagePath = Loader->GetImagePath(ConverterParams->FrameId, ConverterParams->bMipsInSeparateFiles ? CurrentMipLevel : 0);
+			FString ImagePath = Loader->GetImagePath(ConverterParams.FrameId, ConverterParams.bMipsInSeparateFiles ? CurrentMipLevel : 0);
 
-			EReadResult ReadResult = ReadMip(CurrentMipLevel, CurrentTileSelection, OutFrame, SampleConverterPtr, ImagePath, bHasTiles);
+			EReadResult ReadResult = ReadMip(CurrentMipLevel, CurrentTileSelection, OutFrame, ConverterParams, SampleConverterPtr, ImagePath);
 			switch (ReadResult)
 			{
 			case FExrImgMediaReader::Success:
-				if (OutFrame->MipTilesPresent.Find(CurrentMipLevel))
-				{
-					OutFrame->MipTilesPresent[CurrentMipLevel].Include(CurrentTileSelection);
-				}
-				else
-				{
-					OutFrame->MipTilesPresent.Emplace(CurrentMipLevel, CurrentTileSelection);
-				}
+				OutFrame->MipTilesPresent.Include(CurrentMipLevel, CurrentTileSelection);
 				break;
 			case FExrImgMediaReader::Fail:
 				{
@@ -329,7 +311,7 @@ bool FExrImgMediaReaderGpu::ReadFrame(int32 FrameId, const TMap<int32, FImgMedia
 					// Fall back to CPU.
 					bFallBackToCPU = true;
 
-					return FExrImgMediaReader::ReadFrame(ConverterParams->FrameId, InMipTiles, OutFrame);
+					return FExrImgMediaReader::ReadFrame(ConverterParams.FrameId, InMipTiles, OutFrame);
 				}
 				break;
 			case FExrImgMediaReader::Cancelled:
@@ -346,28 +328,29 @@ bool FExrImgMediaReaderGpu::ReadFrame(int32 FrameId, const TMap<int32, FImgMedia
 		}
 		
 		// Create viewport(s) with all mip/tiles present
-		for (const TPair<int32, FImgMediaTileSelection>& TilesPerMip : OutFrame->MipTilesPresent)
+		FScopeLock Lock(&OutFrame->MipTilesPresent.CriticalSection);
+		for (const TPair<int32, FImgMediaTileSelection>& TilesPerMip : OutFrame->MipTilesPresent.GetDataUnsafe())
 		{
 			const FImgMediaTileSelection& CurrentTileSelection = TilesPerMip.Value;
 			const int32 CurrentMipLevel = TilesPerMip.Key;
 
 			// Skip this viewport since we don't have anything to render.
-			if (!SampleConverterPtr->GetMipLevelBuffer(CurrentMipLevel))
+			if (!SampleConverterPtr->HasMipLevelBuffer(CurrentMipLevel))
 			{
 				continue;
 			}
 
 			const int32 MipLevelDiv = 1 << CurrentMipLevel;
-			FIntPoint CurrentMipDim = ConverterParams->FullResolution / MipLevelDiv;
+			FIntPoint CurrentMipDim = ConverterParams.FullResolution / MipLevelDiv;
 
-			TArray<FIntRect>& Viewports = ConverterParams->Viewports.Add(CurrentMipLevel);
+			TArray<FIntRect>& Viewports = ConverterParams.Viewports.Add(CurrentMipLevel);
 			for (const FIntRect& TileRegion : CurrentTileSelection.GetVisibleRegions())
 			{
 				FIntRect Viewport;
-				if (bHasTiles || ConverterParams->bCustomExr)
+				if (ConverterParams.FrameInfo.bHasTiles || ConverterParams.bCustomExr)
 				{
-					Viewport.Min = FIntPoint(ConverterParams->TileDimWithBorders.X * TileRegion.Min.X, ConverterParams->TileDimWithBorders.Y * TileRegion.Min.Y);
-					Viewport.Max = FIntPoint(ConverterParams->TileDimWithBorders.X * TileRegion.Max.X, ConverterParams->TileDimWithBorders.Y * TileRegion.Max.Y);
+					Viewport.Min = FIntPoint(ConverterParams.TileDimWithBorders.X * TileRegion.Min.X, ConverterParams.TileDimWithBorders.Y * TileRegion.Min.Y);
+					Viewport.Max = FIntPoint(ConverterParams.TileDimWithBorders.X * TileRegion.Max.X, ConverterParams.TileDimWithBorders.Y * TileRegion.Max.Y);
 					Viewport.Clip(FIntRect(FIntPoint::ZeroValue, CurrentMipDim));
 				}
 				else
@@ -380,10 +363,11 @@ bool FExrImgMediaReaderGpu::ReadFrame(int32 FrameId, const TMap<int32, FImgMedia
 		}
 	}
 
-	OutFrame->Format = ConverterParams->FrameInfo.NumChannels <= 3 ? EMediaTextureSampleFormat::FloatRGB : EMediaTextureSampleFormat::FloatRGBA;
-	OutFrame->Stride = ConverterParams->FullResolution.X * ConverterParams->PixelSize;
+	OutFrame->Format = ConverterParams.FrameInfo.NumChannels <= 3 ? EMediaTextureSampleFormat::FloatRGB : EMediaTextureSampleFormat::FloatRGBA;
+	OutFrame->Stride = ConverterParams.FullResolution.X * ConverterParams.PixelSize;
+	
+	SampleConverterPtr->SetParams(ConverterParams);
 	
-	ConverterParams->Viewports.KeySort(TLess<int32>());
 	CreateSampleConverterCallback(SampleConverterPtr);
 
 	UE_LOG(LogImgMedia, Verbose, TEXT("Reader %p: Read Pixels Complete. %i"), this, FrameId);
@@ -489,7 +473,7 @@ SIZE_T FExrImgMediaReaderGpu::GetBufferSize(const FIntPoint& Dim, int32 NumChann
 
 void FExrImgMediaReaderGpu::CreateSampleConverterCallback(TSharedPtr<FExrMediaTextureSampleConverter, ESPMode::ThreadSafe> SampleConverter)
 {
-	auto RenderThreadSwizzler = [] (FRHICommandListImmediate& RHICmdList, FTexture2DRHIRef RenderTargetTextureRHI, TMap<int32, FStructuredBufferPoolItemSharedPtr>& MipBuffers, const FSampleConverterParameters& ConverterParams)->bool
+	auto RenderThreadSwizzler = [] (FRHICommandListImmediate& RHICmdList, FTexture2DRHIRef RenderTargetTextureRHI, TMap<int32, FStructuredBufferPoolItemSharedPtr>& MipBuffers, const FSampleConverterParameters ConverterParams)->bool
 	{
 		TRACE_CPUPROFILER_EVENT_SCOPE_TEXT(*FString::Printf(TEXT("ExrReaderGpu.Convert %d"), ConverterParams.FrameId));
 		SCOPED_DRAW_EVENT(RHICmdList, FExrImgMediaReaderGpu_Convert);
@@ -633,7 +617,7 @@ FStructuredBufferPoolItemSharedPtr FExrImgMediaReaderGpu::AllocateGpuBufferFromP
 {
 	TRACE_CPUPROFILER_EVENT_SCOPE_TEXT(*FString::Printf(TEXT("ExrReaderGpu.AllocBuffer")));
 	TRACE_CPUPROFILER_EVENT_SCOPE_TEXT(*FString::Printf(TEXT("ExrReaderGpu.AllocBuffer %d"), AllocSize));
-	TWeakPtr<FExrImgMediaReaderGpu, ESPMode::ThreadSafe> WeakReaderPtr = AsShared();
+	TWeakPtr<FExrImgMediaReaderGpu, ESPMode::ThreadSafe> WeakReaderPtr = AsWeak();
 
 	// This function is attached to the shared pointer and is used to return any allocated memory to staging pool.
 	auto BufferDeleter = [WeakReaderPtr, AllocSize](FStructuredBufferPoolItem* ObjectToDelete) {
@@ -655,7 +639,7 @@ FStructuredBufferPoolItemSharedPtr FExrImgMediaReaderGpu::AllocateGpuBufferFromP
 	FStructuredBufferPoolItemSharedPtr AllocatedBuffer;
 
 	{
-		FScopeLock ScopeLock(&AllocatorCriticalSecion);
+		FScopeLock ScopeLock(&MemoryPoolCriticalSection);
 		FStructuredBufferPoolItem** FoundBuffer = MemoryPool.Find(AllocSize);
 		if (FoundBuffer)
 		{
@@ -674,14 +658,13 @@ FStructuredBufferPoolItemSharedPtr FExrImgMediaReaderGpu::AllocateGpuBufferFromP
 
 		{
 			AllocatedBuffer = MakeShareable(new FStructuredBufferPoolItem(), MoveTemp(BufferDeleter));
-			AllocatedBuffer->Reader = AsShared();
 
 			// Allocate and unlock the structured buffer on render thread.
-			ENQUEUE_RENDER_COMMAND(CreatePooledBuffer)([AllocatedBuffer, AllocSize, InitDoneEvent, this](FRHICommandListImmediate& RHICmdList)
+			ENQUEUE_RENDER_COMMAND(CreatePooledBuffer)([AllocatedBuffer, AllocSize, InitDoneEvent](FRHICommandListImmediate& RHICmdList)
 			{
 				TRACE_CPUPROFILER_EVENT_SCOPE_TEXT(*FString::Printf(TEXT("ExrReaderGpu.AllocBuffer_RenderThread")));
 				TRACE_CPUPROFILER_EVENT_SCOPE_TEXT(*FString::Printf(TEXT("ExrReaderGpu.AllocBuffer_RenderThread %d"), AllocSize));
-				FScopeLock ScopeLock(&AllocatorCriticalSecion);
+				
 				SCOPED_GPU_STAT(RHICmdList, ExrImgMediaReaderGpu_AllocateBuffer);
 				SCOPED_DRAW_EVENT(RHICmdList, FExrImgMediaReaderGpu_AllocateBuffer);
 
@@ -722,7 +705,7 @@ FStructuredBufferPoolItemSharedPtr FExrImgMediaReaderGpu::AllocateGpuBufferFromP
 
 void FExrImgMediaReaderGpu::ReturnGpuBufferToPool(uint32 AllocSize, FStructuredBufferPoolItem* Buffer)
 {
-	FScopeLock ScopeLock(&AllocatorCriticalSecion);
+	FScopeLock ScopeLock(&MemoryPoolCriticalSection);
 	TRACE_CPUPROFILER_EVENT_SCOPE_TEXT(*FString::Printf(TEXT("ExrReaderGpu.ReturnPoolItem")));
 	MemoryPool.Add(AllocSize, Buffer);
 }
@@ -737,7 +720,7 @@ bool FExrMediaTextureSampleConverter::Convert(FTexture2DRHIRef& InDstTexture, co
 	bool bExecutionSuccessful = false;
 	if (ConvertExrBufferCallback.IsBound())
 	{
-		bExecutionSuccessful = ConvertExrBufferCallback.Execute(FRHICommandListExecutor::GetImmediateCommandList(), InDstTexture, MipBuffers, *ConverterParams.Get());
+		bExecutionSuccessful = ConvertExrBufferCallback.Execute(FRHICommandListExecutor::GetImmediateCommandList(), InDstTexture, MipBuffers, GetParams());
 	}
 	return bExecutionSuccessful;
 }
diff --git a/Engine/Plugins/Media/ImgMedia/Source/ImgMedia/Private/Readers/ExrImgMediaReaderGpu.h b/Engine/Plugins/Media/ImgMedia/Source/ImgMedia/Private/Readers/ExrImgMediaReaderGpu.h
--- a/Engine/Plugins/Media/ImgMedia/Source/ImgMedia/Private/Readers/ExrImgMediaReaderGpu.h
+++ b/Engine/Plugins/Media/ImgMedia/Source/ImgMedia/Private/Readers/ExrImgMediaReaderGpu.h
@@ -6,6 +6,7 @@
 
 #if defined(PLATFORM_WINDOWS) && PLATFORM_WINDOWS
 
+#include "Containers/SortedMap.h"
 #include "IImgMediaReader.h"
 #include "IMediaTextureSampleConverter.h"
 
@@ -35,11 +36,6 @@ struct FStructuredBufferPoolItem
 	*/
 	FShaderResourceViewRHIRef ShaderResourceView;
 
-	/**
-	* Keep track of our reader in case it gets destroyed.
-	*/
-	TWeakPtr<FExrImgMediaReaderGpu, ESPMode::ThreadSafe> Reader;
-
 	/**
 	* Destructure to clean up render resources
 	*/
@@ -102,9 +98,9 @@ protected:
 		( const int32 CurrentMipLevel
 		, const FImgMediaTileSelection& CurrentTileSelection
 		, TSharedPtr<FImgMediaFrame, ESPMode::ThreadSafe> OutFrame
+		, FSampleConverterParameters& ConverterParams
 		, TSharedPtr<FExrMediaTextureSampleConverter, ESPMode::ThreadSafe> SampleConverter
-		, const FString& ImagePath
-		, bool bHasTiles);
+		, const FString& ImagePath);
 
 public:
 
@@ -119,7 +115,7 @@ public:
 private:
 
 	/** A critical section used for memory allocation and pool management. */
-	FCriticalSection AllocatorCriticalSecion;
+	FCriticalSection MemoryPoolCriticalSection;
 
 	/** Memory pool from where we are allowed to take buffers. */
 	TMultiMap<uint32, FStructuredBufferPoolItem*> MemoryPool;
@@ -153,7 +149,7 @@ struct FSampleConverterParameters
 	FIntPoint TileDimWithBorders;
 
 	/** Used for rendering tiles in bulk regions per mip level. */
-	TMap<int32, TArray<FIntRect>> Viewports;
+	TSortedMap<int32, TArray<FIntRect>> Viewports;
 
 #if defined(PLATFORM_WINDOWS) && PLATFORM_WINDOWS
 	/** Contain information about individual tiles. Used to convert buffer data into a 2D Texture.
@@ -179,7 +175,7 @@ struct FSampleConverterParameters
 };
 
 
-FUNC_DECLARE_DELEGATE(FExrConvertBufferCallback, bool, FRHICommandListImmediate& /*RHICmdList*/, FTexture2DRHIRef /*RenderTargetTextureRHI*/, TMap<int32, FStructuredBufferPoolItemSharedPtr>& /*MipBuffers*/, const FSampleConverterParameters & /*ConverterParams*/)
+FUNC_DECLARE_DELEGATE(FExrConvertBufferCallback, bool, FRHICommandListImmediate& /*RHICmdList*/, FTexture2DRHIRef /*RenderTargetTextureRHI*/, TMap<int32, FStructuredBufferPoolItemSharedPtr>& /*MipBuffers*/, const FSampleConverterParameters /*ConverterParams*/)
 
 class FExrMediaTextureSampleConverter: public IMediaTextureSampleConverter
 {
@@ -194,44 +190,49 @@ public:
 		ConvertExrBufferCallback = Callback;
 
 		// Copy mip buffers to be used for rendering. 
+		FScopeLock Lock(&MipBufferCriticalSection);
 		MipBuffersRenderThread = MipBuffers;
 	};
 
-	void LockMipBuffers()
+	bool HasMipLevelBuffer(int32 RequestedMipLevel) const
 	{
-		MipBufferCriticalSection.Lock();
+		FScopeLock Lock(&MipBufferCriticalSection);
+		return MipBuffers.Contains(RequestedMipLevel);
 	}
 
-	void UnlockMipBuffers()
+	FStructuredBufferPoolItemSharedPtr GetOrCreateMipLevelBuffer(int32 RequestedMipLevel, TFunction<FStructuredBufferPoolItemSharedPtr()> AllocatorFunc)
 	{
-		MipBufferCriticalSection.Unlock();
-	}
-
-	FStructuredBufferPoolItemSharedPtr GetMipLevelBuffer(int32 RequestedMipLevel)
-	{
-		if (MipBuffers.Contains(RequestedMipLevel))
+		FScopeLock Lock(&MipBufferCriticalSection);
+		if (FStructuredBufferPoolItemSharedPtr* BufferDataPtr = MipBuffers.Find(RequestedMipLevel))
 		{
-			return *MipBuffers.Find(RequestedMipLevel);
+			return *BufferDataPtr;
 		}
 
-		return nullptr;
+		return MipBuffers.Add(RequestedMipLevel, AllocatorFunc());
+	}
+	
+	FSampleConverterParameters GetParams()
+	{
+		FScopeLock ScopeLock(&ParamsCriticalSection);
+		return ConverterParams;
 	}
 
-	void SetMipLevelBuffer(int32 RequestedMipLevel, FStructuredBufferPoolItemSharedPtr Buffer)
+	void SetParams(const FSampleConverterParameters& InParams)
 	{
-		check(!MipBuffers.Contains(RequestedMipLevel));
-		MipBuffers.Add(RequestedMipLevel, Buffer);
+		FScopeLock ScopeLock(&ParamsCriticalSection);
+		ConverterParams = InParams;
 	}
-public:
-	/** These are all required parameters to convert the buffer into texture successfully. */
-	TSharedPtr<FSampleConverterParameters> ConverterParams = MakeShared<FSampleConverterParameters>();
 
 private:
+	/** These are all required parameters to convert the buffer into texture successfully. */
+	mutable FCriticalSection ParamsCriticalSection;
+	FSampleConverterParameters ConverterParams;
+
 	FCriticalSection ConverterCallbacksCriticalSection;
 	FExrConvertBufferCallback ConvertExrBufferCallback;
 
 	/** Lock to be used exclusively on reader threads.*/
-	FCriticalSection MipBufferCriticalSection;
+	mutable FCriticalSection MipBufferCriticalSection;
 
 	/** An array of structured buffers that are big enough to fully contain corresponding mip levels used by reader threads and transferred into MipBuffersRenderThread. */
 	TMap<int32,FStructuredBufferPoolItemSharedPtr> MipBuffers;
diff --git a/Engine/Plugins/Media/ImgMedia/Source/ImgMedia/Private/Readers/GenericImgMediaReader.cpp b/Engine/Plugins/Media/ImgMedia/Source/ImgMedia/Private/Readers/GenericImgMediaReader.cpp
--- a/Engine/Plugins/Media/ImgMedia/Source/ImgMedia/Private/Readers/GenericImgMediaReader.cpp
+++ b/Engine/Plugins/Media/ImgMedia/Source/ImgMedia/Private/Readers/GenericImgMediaReader.cpp
@@ -167,11 +167,11 @@ bool FGenericImgMediaReader::ReadFrame(int32 FrameId, const TMap<int32, FImgMedi
 					}
 
 					void* Buffer = FMemory::Malloc(AllocSize, PLATFORM_CACHE_LINE_SIZE);
-					OutFrame->Info = Info;
+					OutFrame->SetInfo(Info);
 					OutFrame->Data = MakeShareable(Buffer, [](void* ObjectToDelete) { FMemory::Free(ObjectToDelete); });
 					OutFrame->MipTilesPresent.Reset();
 					OutFrame->Format = EMediaTextureSampleFormat::CharBGRA;
-					OutFrame->Stride = OutFrame->Info.Dim.X * 4;
+					OutFrame->Stride = Info.Dim.X * 4;
 				}
 
 				// Copy data to our buffer with the right mip level offset
diff --git a/Engine/Plugins/Media/ImgMedia/Source/ImgMedia/Private/Readers/IImgMediaReader.h b/Engine/Plugins/Media/ImgMedia/Source/ImgMedia/Private/Readers/IImgMediaReader.h
--- a/Engine/Plugins/Media/ImgMedia/Source/ImgMedia/Private/Readers/IImgMediaReader.h
+++ b/Engine/Plugins/Media/ImgMedia/Source/ImgMedia/Private/Readers/IImgMediaReader.h
@@ -54,46 +54,294 @@ struct FImgMediaFrameInfo
 };
 
 
+/**
+ * A thread-safe container for tile selections per mip level.
+ */
+struct FImgMediaMipTiles final
+{
+	/** Default constructor. */
+	FImgMediaMipTiles() = default;
+	/** Default destructor. */
+	~FImgMediaMipTiles() = default;
+	
+	/** Move constructor. */
+	FImgMediaMipTiles(FImgMediaMipTiles&& Other)
+	{
+		InternalMove(MoveTemp(Other));
+	}
+	
+	/** Copy constructor. */
+	FImgMediaMipTiles(const FImgMediaMipTiles& Other)
+	{
+		InternalCopy(Other);
+	}
+
+	/** Move assignment operator. */
+	FImgMediaMipTiles& operator=(FImgMediaMipTiles&& Other)
+	{
+		InternalMove(MoveTemp(Other));
+		return *this;
+	}
+
+	/** Copy assignment operator. */
+	FImgMediaMipTiles& operator=(const FImgMediaMipTiles& Other)
+	{
+		InternalCopy(Other);
+		return *this;
+	}
+
+	/**
+	 * Contains the specified mip level.
+	 * 
+	 * @param InMipLevel mip level
+	 * @return true when mip level is present.
+	*/
+	bool Contains(const int32 InMipLevel) const
+	{
+		FScopeLock Lock(&CriticalSection);
+		return MipTiles.Contains(InMipLevel);
+	}
+
+	/**
+	 * Check if the container is empty.
+	 *
+	 * @return true when no mips/tiles are present.
+	*/
+	bool IsEmpty() const
+	{
+		FScopeLock Lock(&CriticalSection);
+		return MipTiles.IsEmpty();
+	}
+
+	/** Reset the container. */
+	void Reset()
+	{
+		FScopeLock Lock(&CriticalSection);
+		return MipTiles.Reset();
+	}
+
+	/**
+	 * Check if the existing tiles contain all of the requested ones. (Is the container a superset of requested?)
+	 *
+	 * @param InRequestedMipTiles Requested tile selection per mip level.
+	 * @return true if the selection is entirely contained, or the request was empty.
+	 */
+	bool ContainsMipTiles(const TMap<int32, FImgMediaTileSelection>& InRequestedMipTiles)
+	{
+		for (auto Iter = InRequestedMipTiles.CreateConstIterator(); Iter; ++Iter)
+		{
+			if (!ContainsTiles(Iter.Key(), Iter.Value()))
+			{
+				return false;
+			}
+		}
+
+		return true;
+	}
+
+	/**
+	 * Check if the existing tiles contain all of the requested ones, at the specified mip level.
+	 * 
+	 * @param InMipLevel Specified mip level.
+	 * @param InSelection Requested tile selection.
+	 * @return true if the selection is contained
+	 */
+	bool ContainsTiles(const int32 InMipLevel, const FImgMediaTileSelection& InSelection)
+	{
+		FScopeLock Lock(&CriticalSection);
+
+		if (const FImgMediaTileSelection* CachedSelection = MipTiles.Find(InMipLevel))
+		{
+			// Check if requested tile selection is not present.
+			return CachedSelection->Contains(InSelection);
+		}
+
+		// Requested mip level is not present.
+		return false;
+	}
+
+	/**
+	 * Include the selection at the specified mip level.
+	 *
+	 * @param InMipLevel Specified mip level.
+	 * @param InSelection New tile selection.
+	 */
+	void Include(const int32 InMipLevel, const FImgMediaTileSelection& InSelection)
+	{
+		FScopeLock Lock(&CriticalSection);
+
+		if (FImgMediaTileSelection* Tiles = MipTiles.Find(InMipLevel))
+		{
+			Tiles->Include(InSelection);
+		}
+		else
+		{
+			MipTiles.Emplace(InMipLevel, InSelection);
+		}
+	}
+
+	/**
+	 * Emplace the selection at the specified mip level.
+	 *
+	 * @param InMipLevel Specified mip level.
+	 * @param InSelection New tile selection.
+	 */
+	void Emplace(const int32 InMipLevel, const FImgMediaTileSelection& InSelection)
+	{
+		FScopeLock Lock(&CriticalSection);
+
+		MipTiles.Emplace(InMipLevel, InSelection);
+	}
+
+	/**
+	 * Get a copy of the selection at the specified mip level.
+	 * 
+	 * @param InMipLevel Specified mip level.
+	 * @param OutSelection Returned tile selection.
+	 * @return true if mip level was found
+	 */
+	bool GetSelection(const int32 InMipLevel, FImgMediaTileSelection& OutSelection) const
+	{
+		FScopeLock Lock(&CriticalSection);
+
+		if (const FImgMediaTileSelection* Selection = MipTiles.Find(InMipLevel))
+		{
+			OutSelection = *Selection;
+			return true;
+		}
+
+		return false;
+	}
+
+	/**
+	 * Returns a calculated list of contiguous visible tile regions. Only provides regions for the missing tiles if
+	 * CurrentTileSelection is specified.
+	 *
+	 * @param InMipLevel Specified mip level.
+	 * @param InSelection Existing tile selection.
+	 * @param OutRegions Visible tile regions array.
+	 * @return true if mip level was found
+	 */
+	bool GetVisibleRegions(const int32 InMipLevel, const FImgMediaTileSelection& InSelection, TArray<FIntRect>& OutRegions) const
+	{
+		FScopeLock Lock(&CriticalSection);
+
+		if (const FImgMediaTileSelection* Selection = MipTiles.Find(InMipLevel))
+		{
+			OutRegions = Selection->GetVisibleRegions(&InSelection);
+			return true;
+		}
+
+		return false;
+	}
+
+	/** Thread-unsafe data getter: the caller is responsible for locking. */
+	TMap<int32, FImgMediaTileSelection>& GetDataUnsafe()
+	{
+		return MipTiles;
+	}
+
+	/** Thread-unsafe data getter: the caller is responsible for locking. */
+	const TMap<int32, FImgMediaTileSelection>& GetDataUnsafe() const
+	{
+		return MipTiles;
+	}
+
+public:
+
+	/** Critical section to lock when using unsafe methods. */
+	mutable FCriticalSection CriticalSection;
+
+private:
+	void InternalMove(FImgMediaMipTiles&& Other)
+	{
+		FScopeLock LockOther(&Other.CriticalSection);
+		FScopeLock Lock(&CriticalSection);
+		MipTiles = MoveTemp(Other.MipTiles);
+	}
+
+	void InternalCopy(const FImgMediaMipTiles& Other)
+	{
+		FScopeLock LockOther(&Other.CriticalSection);
+		FScopeLock Lock(&CriticalSection);
+		MipTiles = Other.MipTiles;
+	}
+
+	/** Tiles present per mip level. */
+	TMap<int32, FImgMediaTileSelection> MipTiles;
+};
+
 /**
  * A single frame of an image sequence.
  */
-struct FImgMediaFrame
+struct FImgMediaFrame final
 {
+	/** Default constructor. */
+	FImgMediaFrame() = default;
+
+	/** Default destructor. */
+	~FImgMediaFrame() = default;
+
+	/** Copy constructor. */
+	FImgMediaFrame(const FImgMediaFrame& Other)
+		: Data(Other.Data)
+		, Format(Other.Format.load())
+		, MipTilesPresent(Other.MipTilesPresent)
+		, NumTilesRead(Other.NumTilesRead.load())
+		, Stride(Other.Stride.load())
+		, Info() // copied below
+		, SampleConverter() // copied below
+	{
+		{
+			FScopeLock LockOther(&Other.SampleConverterCriticalSection);
+			SampleConverter = Other.SampleConverter;
+		}
+
+		{
+			FScopeLock LockOther(&Other.InfoCriticalSection);
+			Info = Other.Info;
+		}
+	}
+
+	/** Copy assignment operator. */
+	FImgMediaFrame& operator=(const FImgMediaFrame& Other)
+	{
+		if (this != &Other)
+		{
+			this->~FImgMediaFrame();
+			new(this) FImgMediaFrame(Other);
+		}
+		return *this;
+	}
+
+	FImgMediaFrame(FImgMediaFrame&& Other) = delete;
+	FImgMediaFrame& operator=(FImgMediaFrame&& Other) = delete;
+
 	/** The frame's data. */
 	TSharedPtr<void, ESPMode::ThreadSafe> Data;
 
 	/** The frame's sample format. */
-	EMediaTextureSampleFormat Format;
-
-	/** Additional information about the frame. */
-	FImgMediaFrameInfo Info;
+	std::atomic<EMediaTextureSampleFormat> Format;
 
 	/** Tiles present per mip level. */
-	TMap<int32, FImgMediaTileSelection> MipTilesPresent;
+	FImgMediaMipTiles MipTilesPresent;
 
 	/** Total number of tiles read. */
-	int32 NumTilesRead = 0;
+	std::atomic<int32> NumTilesRead = 0;
 
 	/** The frame's horizontal stride (in bytes). */
-	uint32 Stride = 0;
+	std::atomic<uint32> Stride = 0;
 
 	// This should only be used if you do not need to make changes to sample converter.
 	virtual IMediaTextureSampleConverter* GetSampleConverter()
 	{
-		if (!SampleConverter.IsValid())
-		{
-			return nullptr;
-		}
 		return SampleConverter.Get();
 	};
 
-	/** Virtual non trivial destructor. */
-	virtual ~FImgMediaFrame() {};
-
 	template <class T>
 	TSharedPtr<T, ESPMode::ThreadSafe> GetOrCreateSampleConverter()
 	{
-		FScopeLock ScopedLock(SampleConverterCriticalSection.Get());
+		FScopeLock ScopedLock(&SampleConverterCriticalSection);
 		if (!SampleConverter.IsValid())
 		{
 			SampleConverter = MakeShared<T>();
@@ -101,9 +349,58 @@ struct FImgMediaFrame
 		return StaticCastSharedPtr<T>(SampleConverter);
 	}
 
+	/** Thread-safe info getter. */
+	FImgMediaFrameInfo GetInfo() const
+	{
+		FScopeLock ScopedLock(&InfoCriticalSection);
+		return Info;
+	}
+
+	/** Thread-safe info setter. */
+	void SetInfo(FImgMediaFrameInfo InInfo)
+	{
+		FScopeLock ScopedLock(&InfoCriticalSection);
+		Info = InInfo;
+	}
+
+	/** Thread-safe framerate getter. */
+	FFrameRate GetFrameRate() const
+	{
+		FScopeLock ScopedLock(&InfoCriticalSection);
+		return Info.FrameRate;
+	}
+
+	/** Thread-safe dimensions getter. */
+	FIntPoint GetDim() const
+	{
+		FScopeLock ScopedLock(&InfoCriticalSection);
+		return Info.Dim;
+	}
+
+	/** Thread-safe sRGB getter. */
+	bool IsOutputSrgb() const
+	{
+		FScopeLock ScopedLock(&InfoCriticalSection);
+		return Info.Srgb;
+	}
+
+	/** Thread-safe uncompressed size. */
+	SIZE_T GetUncompressedSize() const
+	{
+		FScopeLock ScopedLock(&InfoCriticalSection);
+		return Info.UncompressedSize;
+	}
+
 private:
-	/** Lock to be used exclusively on reader threads.*/
-	TSharedPtr<FCriticalSection, ESPMode::ThreadSafe> SampleConverterCriticalSection = MakeShared<FCriticalSection>();
+
+	/** Information critical section. */
+	mutable FCriticalSection InfoCriticalSection;
+
+	/** Additional information about the frame. */
+	FImgMediaFrameInfo Info;
+
+	/** Lock to be used exclusively on reader threads. Note: this protects only the creation of the sample converter. */
+	mutable FCriticalSection SampleConverterCriticalSection;
 
 	/** Sample converter is used by Media Texture Resource to convert the texture or data. */
 	TSharedPtr<IMediaTextureSampleConverter, ESPMode::ThreadSafe> SampleConverter;
