ImgMedia: Reset cached view info on begin frame event.

#rnx

diff --git a/Engine/Plugins/Media/ImgMedia/Source/ImgMedia/Private/ImgMediaSceneViewExtension.cpp b/Engine/Plugins/Media/ImgMedia/Source/ImgMedia/Private/ImgMediaSceneViewExtension.cpp
--- a/Engine/Plugins/Media/ImgMedia/Source/ImgMedia/Private/ImgMediaSceneViewExtension.cpp
+++ b/Engine/Plugins/Media/ImgMedia/Source/ImgMedia/Private/ImgMediaSceneViewExtension.cpp
@@ -37,24 +37,19 @@ static TAutoConsoleVariable<bool> CVarImgMediaProcessTilesInnerOnly(
 FImgMediaSceneViewExtension::FImgMediaSceneViewExtension(const FAutoRegister& AutoReg)
 	: FSceneViewExtensionBase(AutoReg)
 	, CachedViewInfos()
-	, LastFrameNumber(0)
 {
+	OnBeginFrameDelegate = FCoreDelegates::OnBeginFrame.AddRaw(this, &FImgMediaSceneViewExtension::ResetViewInfoCache);
+}
+
+FImgMediaSceneViewExtension::~FImgMediaSceneViewExtension()
+{
+	FCoreDelegates::OnBeginFrame.Remove(OnBeginFrameDelegate);
 }
 
 void FImgMediaSceneViewExtension::BeginRenderViewFamily(FSceneViewFamily& InViewFamily)
 {
 	TRACE_CPUPROFILER_EVENT_SCOPE(FImgMediaSceneViewExtension::BeginRenderViewFamily);
 
-	/**
-	* NOTE: Scene captures call this function after the primary `BeginRenderViewFamily` with the same frame number.
-	* Therefore, the view infos we cache here will be correctly kept until the next frame.
-	*/
-	if (LastFrameNumber != InViewFamily.FrameNumber)
-	{
-		CachedViewInfos.Reset();
-		LastFrameNumber = InViewFamily.FrameNumber;
-	}
-
 	for (const FSceneView* View : InViewFamily.Views)
 	{
 		if (View != nullptr)
@@ -156,4 +151,10 @@ void FImgMediaSceneViewExtension::CacheViewInfo(FSceneViewFamily& InViewFamily,
 
 	CachedViewInfos.Add(MoveTemp(Info));
 }
+
+void FImgMediaSceneViewExtension::ResetViewInfoCache()
+{
+	CachedViewInfos.Reset();
+}
+
 #undef LOCTEXT_NAMESPACE
diff --git a/Engine/Plugins/Media/ImgMedia/Source/ImgMedia/Public/ImgMediaSceneViewExtension.h b/Engine/Plugins/Media/ImgMedia/Source/ImgMedia/Public/ImgMediaSceneViewExtension.h
--- a/Engine/Plugins/Media/ImgMedia/Source/ImgMedia/Public/ImgMediaSceneViewExtension.h
+++ b/Engine/Plugins/Media/ImgMedia/Source/ImgMedia/Public/ImgMediaSceneViewExtension.h
@@ -37,6 +37,7 @@ class FImgMediaSceneViewExtension final : public FSceneViewExtensionBase
 {
 public:
 	FImgMediaSceneViewExtension(const FAutoRegister& AutoReg);
+	~FImgMediaSceneViewExtension();
 
 	/**
 	 * Get the cached camera information array, updated on the game thread by BeginRenderViewFamily.
@@ -53,10 +54,13 @@ public:
 private:
 	/** Cache camera view information for the current frame. */
 	void CacheViewInfo(FSceneViewFamily& InViewFamily, const FSceneView& View);
+
+	/** Reset the view info cache. */
+	void ResetViewInfoCache();
 	
 	/** Array of info on each camera used for mipmap calculations. */
 	TArray<FImgMediaViewInfo> CachedViewInfos;
 
-	/** Last received FSceneViewFamily frame number. */
-	uint32 LastFrameNumber;
+	/** FCoreDelegates::OnBeginFrame delegate. */
+	FDelegateHandle OnBeginFrameDelegate;
 };
