Media Plate: Added adaptive pole mip upscaling.

diff --git a/Engine/Plugins/Media/ImgMedia/Source/ImgMedia/Internal/Assets/ImgMediaMipMapInfo.h b/Engine/Plugins/Media/ImgMedia/Source/ImgMedia/Internal/Assets/ImgMediaMipMapInfo.h
--- a/Engine/Plugins/Media/ImgMedia/Source/ImgMedia/Internal/Assets/ImgMediaMipMapInfo.h
+++ b/Engine/Plugins/Media/ImgMedia/Source/ImgMedia/Internal/Assets/ImgMediaMipMapInfo.h
@@ -212,7 +212,7 @@ public:
 	UMeshComponent* GetMeshComponent() const;
 
 	/** Get the mip level to upscale (-1 if disabled). */
-	int32 GetMipLevelToUpscale() const;
+	virtual int32 GetMipLevelToUpscale() const;
 
 protected:
 
diff --git a/Engine/Plugins/Media/ImgMedia/Source/ImgMedia/Private/Assets/ImgMediaMipMapInfo.cpp b/Engine/Plugins/Media/ImgMedia/Source/ImgMedia/Private/Assets/ImgMediaMipMapInfo.cpp
--- a/Engine/Plugins/Media/ImgMedia/Source/ImgMedia/Private/Assets/ImgMediaMipMapInfo.cpp
+++ b/Engine/Plugins/Media/ImgMedia/Source/ImgMedia/Private/Assets/ImgMediaMipMapInfo.cpp
@@ -627,6 +627,11 @@ namespace {
 		{
 		}
 
+		virtual int32 GetMipLevelToUpscale() const override
+		{
+			return MipLevelToUpscale;
+		}
+
 		void CalculateVisibleTiles(const TArray<FImgMediaViewInfo>& InViewInfos, const FSequenceInfo& InSequenceInfo, TMap<int32, FImgMediaTileSelection>& VisibleTiles) const override
 		{
 			UMeshComponent* Mesh = MeshComponent.Get();
@@ -655,8 +660,24 @@ namespace {
 			const float PixelDimY = 1.0f / InSequenceInfo.Dim.Y;
 
 			const FVector ApproxTileSizeWS = MeshTransform.GetScale3D() * (UE_TWO_PI * DefaultSphereRadius) / FMath::Max(SequencePartialTileNum.X, SequencePartialTileNum.Y);
+			const int32 TotalNumTiles = FMath::CeilToInt32(SequencePartialTileNum.X * SequencePartialTileNum.Y);
+
 			const float ApproxTileRadiusInWS = 0.5f * UE_SQRT_2 * ApproxTileSizeWS.GetAbsMax();
 
+			// Does user want to upscale specific mip level? 
+			int32 MipLevelToUpscaleExcludingPoles = -1;
+
+			// Does user want to reduce the load at the poles?
+			bool bAdaptivePoleMipUpscaling = false;
+			MipLevelToUpscale = -1;
+			if (Tracker.IsValid())
+			{
+				TSharedPtr<FMediaTextureTrackerObject, ESPMode::ThreadSafe> PinnedTracker = Tracker.Pin();
+				MipLevelToUpscaleExcludingPoles = PinnedTracker->MipLevelToUpscale;
+				bAdaptivePoleMipUpscaling = PinnedTracker->bAdaptivePoleMipUpscaling;
+				MipLevelToUpscale = MipLevelToUpscaleExcludingPoles;
+			}
+
 			for (const FImgMediaViewInfo& ViewInfo : InViewInfos)
 			{
 				if (IsPrimitiveComponentHidden(Mesh->GetPrimitiveSceneId(), ViewInfo))
@@ -671,8 +692,43 @@ namespace {
 				// Approximated UV coordinate for a camera centered inside the sphere
 				FVector2f ViewUV = TransformDirectionWSToSphericalUVs(MeshRange, MeshTransform, ViewInfo.ViewDirection);
 
-				for (int32 TileY = 0; TileY < SequenceTileNum.Y; ++TileY)
+
+
+				// 20 Degrees equates to total 11.1% of a sphere for top and bottom poles.
+				const float PoleEdgeDegrees = 20.f;
+				// Angle from the pole at which upscaling is enabled automatically.
+				const float UpscalingEnabledEdge = 10.f;
+
+				// NumOfTilesToLoadAtThePole - how many tiles does this sphere expect to load without upscaling.
+				// PoleMipBias - automatically calculated value of which mip to upscale.
+				int32 NumOfTilesToLoadAtThePole = 0, PoleMipBias = -1;
+				// This number represents a factor by which we need to reduce the number of tiles loaded.
+				// Ex: Visible tiles at mip 0 at a 20 degree pole = 100. To improve perf we should load 100/TileReductionFactor = 50 tiles.
+				const float TileReductionFactor = 0.5f;
+				if (bAdaptivePoleMipUpscaling)
+				{
+					NumOfTilesToLoadAtThePole = FMath::CeilToInt32(SequencePartialTileNum.X * ((float)SequencePartialTileNum.Y * PoleEdgeDegrees / 180.f))* TileReductionFactor;
+					PoleMipBias = FMath::CeilToInt32(FMath::LogX(4., TotalNumTiles/ NumOfTilesToLoadAtThePole));
+				}
+
+				// Everything above this value gets upscaled.
+				float TipUpscalingEdgePercent = (1.f - PoleEdgeDegrees / 90.f);
+
+				// At which point do we consider that tip is in the view.
+				float TipInTheViewThresholdPercent = (1.f - UpscalingEnabledEdge / 90.f);
+
+				// The following bool indicates if camera is looking at the tip and if we should not load tiles above TipUpscalingEdge
+				bool bPoleTipIsInView = false;
+
+				auto ProcessTileRow = [&](int32 TileY)
 				{
+					float TileVMin = 0, TileVMax = 0;
+					if (bAdaptivePoleMipUpscaling)
+					{
+						TileVMin = ((((float)TileY) / SequencePartialTileNum.Y) - 0.5) * 2.;
+						TileVMax = ((((float)TileY + 1.) / SequencePartialTileNum.Y) - 0.5) * 2.;
+					}
+
 					for (int32 TileX = 0; TileX < SequenceTileNum.X; ++TileX)
 					{
 						const FVector2f TileMinCornerUV = FVector2f((float)TileX, (float)TileY) / SequencePartialTileNum;
@@ -692,6 +748,12 @@ namespace {
 
 						if (ViewFrustum.IntersectSphere(TileLocationWS, CollisionSphereRadius))
 						{
+							if (bAdaptivePoleMipUpscaling && (TileVMin < -TipInTheViewThresholdPercent || TileVMax > TipInTheViewThresholdPercent))
+							{
+								bPoleTipIsInView = true;
+								// Queue upscaling of lower quality mips into this mip.
+								MipLevelToUpscale = FMath::Max(PoleMipBias, MipLevelToUpscaleExcludingPoles);
+							}
 							float CalculatedLevel;
 							FIntVector2 MipLevelRange;
 
@@ -716,6 +778,14 @@ namespace {
 
 								for (int32 Level = MipLevelRange[0]; Level <= MipLevelRange[1]; ++Level)
 								{
+									if (bAdaptivePoleMipUpscaling && Level < PoleMipBias)
+									{
+										// Anything above TipUpscalingEdgePercent will be upscaled.
+										if (bPoleTipIsInView && (TileVMin < -TipUpscalingEdgePercent || TileVMax > TipUpscalingEdgePercent))
+										{
+											continue;
+										}
+									}
 									if (!VisibleTiles.Contains(Level))
 									{
 										VisibleTiles.Emplace(Level, FImgMediaTileSelection::CreateForTargetMipLevel(InSequenceInfo.Dim, InSequenceInfo.TilingDescription.TileSize, Level, false));
@@ -739,6 +809,17 @@ namespace {
 #endif // false
 						}
 					}
+				};
+
+				int32 MiddleRowIndex = FMath::FloorToInt32(((float)SequenceTileNum.Y) / 2.);
+
+				for (int32 TileY = 0; TileY < MiddleRowIndex; ++TileY)
+				{
+					ProcessTileRow(TileY);
+				}
+				for (int32 TileY = SequenceTileNum.Y - 1; TileY >= MiddleRowIndex; TileY--)
+				{
+					ProcessTileRow(TileY);
 				}
 			}
 		}
@@ -770,6 +851,8 @@ namespace {
 		}
 
 		const float DefaultSphereRadius;
+	private:
+		mutable int32 MipLevelToUpscale;
 	};
 
 } //end anonymous namespace
diff --git a/Engine/Plugins/Media/MediaPlate/Source/MediaPlate/Private/MediaPlateComponent.cpp b/Engine/Plugins/Media/MediaPlate/Source/MediaPlate/Private/MediaPlateComponent.cpp
--- a/Engine/Plugins/Media/MediaPlate/Source/MediaPlate/Private/MediaPlateComponent.cpp
+++ b/Engine/Plugins/Media/MediaPlate/Source/MediaPlate/Private/MediaPlateComponent.cpp
@@ -458,6 +458,7 @@ void UMediaPlateComponent::RegisterWithMediaTextureTracker()
 	MediaTextureTrackerObject->VisibleMipsTilesCalculations = VisibleMipsTilesCalculations;
 	MediaTextureTrackerObject->MeshRange = MeshRange;
 	MediaTextureTrackerObject->MipLevelToUpscale = bEnableMipMapUpscaling ? MipLevelToUpscale : -1;
+	MediaTextureTrackerObject->bAdaptivePoleMipUpscaling = bAdaptivePoleMipUpscaling;
 
 	// Add our textures.
 	FMediaTextureTracker& MediaTextureTracker = FMediaTextureTracker::Get();
diff --git a/Engine/Plugins/Media/MediaPlate/Source/MediaPlate/Public/MediaPlateComponent.h b/Engine/Plugins/Media/MediaPlate/Source/MediaPlate/Public/MediaPlateComponent.h
--- a/Engine/Plugins/Media/MediaPlate/Source/MediaPlate/Public/MediaPlateComponent.h
+++ b/Engine/Plugins/Media/MediaPlate/Source/MediaPlate/Public/MediaPlateComponent.h
@@ -366,6 +366,10 @@ private:
 	UPROPERTY(EditAnywhere, AdvancedDisplay, Category = "MediaPlate", meta = (EditCondition = "bEnableMipMapUpscaling", AllowPrivateAccess = true, UIMin = "0", UIMax = "16"))
 	int32 MipLevelToUpscale = 16;
 
+	/** If true then Media Plate will attempt to load and upscale lower quality mips and display those at the poles (Sphere object only). */
+	UPROPERTY(EditAnywhere, AdvancedDisplay, Category = "MediaPlate", meta = (EditCondition = "VisibleMipsTilesCalculations == EMediaTextureVisibleMipsTiles::Sphere", AllowPrivateAccess = true))
+	bool bAdaptivePoleMipUpscaling = false;
+
 	/** If > 0, then this is the aspect ratio of our screen and 
 	 * letterboxes will be added if the media is smaller than the screen. */
 	UPROPERTY()
diff --git a/Engine/Source/Runtime/MediaAssets/Public/MediaTextureTracker.h b/Engine/Source/Runtime/MediaAssets/Public/MediaTextureTracker.h
--- a/Engine/Source/Runtime/MediaAssets/Public/MediaTextureTracker.h
+++ b/Engine/Source/Runtime/MediaAssets/Public/MediaTextureTracker.h
@@ -34,6 +34,8 @@ struct FMediaTextureTrackerObject
 	FVector2D MeshRange;
 	/* Mip level to upscale into lower quality mips. All levels including and above the specified value will be fully read. */
 	int32 MipLevelToUpscale = -1;
+	/* When this is enabled Media Plate will automatically try to figure out the best mip to upscale to reduce the load. */
+	bool bAdaptivePoleMipUpscaling = false;
 };
 
 /**
