ImgMedia: OpenExr's RgbaInputFile reads all exr files as RGBA 16 bit per channel exr, however cache was using uncompressed exr size rather than its RGBA interpretation which was creating invalid cache values.
- additionally updated the way uncompressed size is calculated.

diff --git a/Engine/Plugins/Media/ImgMedia/Source/ImgMedia/Private/Readers/ExrImgMediaReader.cpp b/Engine/Plugins/Media/ImgMedia/Source/ImgMedia/Private/Readers/ExrImgMediaReader.cpp
--- a/Engine/Plugins/Media/ImgMedia/Source/ImgMedia/Private/Readers/ExrImgMediaReader.cpp
+++ b/Engine/Plugins/Media/ImgMedia/Source/ImgMedia/Private/Readers/ExrImgMediaReader.cpp
@@ -205,7 +205,10 @@ bool FExrImgMediaReader::ReadFrame(int32 FrameId, const TMap<int32, FImgMediaTil
 		}
 
 		// allocate frame buffer
-		SIZE_T BufferSize = GetMipBufferTotalSize(Dim);
+		// FRgbaInputFile loads any exr as rgba 16 bit per channel.
+		// UncompressedSize is used for cache, therefore it needs to be set to the actual size of the buffer.
+		FrameInfo.UncompressedSize = GetMipBufferTotalSize(Dim, FrameInfo.NumMipLevels > 1);
+		SIZE_T BufferSize = FrameInfo.UncompressedSize;
 		void* Buffer = FMemory::Malloc(BufferSize, PLATFORM_CACHE_LINE_SIZE);
 
 		auto BufferDeleter = [BufferSize](void* ObjectToDelete) {
@@ -274,7 +277,7 @@ bool FExrImgMediaReader::ReadFrame(int32 FrameId, const TMap<int32, FImgMediaTil
 					ConverterParams.FrameId = FrameId;
 
 					TArray<UE::Math::TIntPoint<int64>> OutBufferRegionsToCopy;
-					EReadResult ReadResult = ReadTiles((uint16*)MipDataPtr, GetMipBufferTotalSize(Dim / MipLevelDiv), Image, TileRegions, ConverterParams, CurrentMipLevel, OutBufferRegionsToCopy);
+					EReadResult ReadResult = ReadTiles((uint16*)MipDataPtr, FrameInfo.UncompressedSize / MipLevelDiv, Image, TileRegions, ConverterParams, CurrentMipLevel, OutBufferRegionsToCopy);
 					if (ReadResult != Fail)
 					{
 						OutFrame->MipTilesPresent.Emplace(CurrentMipLevel, CurrentTileSelection);
@@ -456,10 +459,18 @@ void FExrImgMediaReader::SetCustomFormatInfo(bool bInIsCustomFormat, const FIntP
 	bIsCustomFormatTiled = InTileSize.X != 0;
 }
 
-SIZE_T FExrImgMediaReader::GetMipBufferTotalSize(FIntPoint Dim)
+SIZE_T FExrImgMediaReader::GetMipBufferTotalSize(FIntPoint Dim, bool bInHasMips)
 {
-	SIZE_T Size = ((Dim.X * Dim.Y * 4) / 3) * sizeof(uint16) * 4;
-	
+	SIZE_T Size = 0;
+	if (bInHasMips)
+	{
+		Size = ((Dim.X * Dim.Y * 4) / 3) * sizeof(uint16) * 4;
+	}
+	else
+	{
+		Size = (Dim.X * Dim.Y ) * sizeof(uint16) * 4;
+	}
+
 	return Size;
 }
 
diff --git a/Engine/Plugins/Media/ImgMedia/Source/ImgMedia/Private/Readers/ExrImgMediaReader.h b/Engine/Plugins/Media/ImgMedia/Source/ImgMedia/Private/Readers/ExrImgMediaReader.h
--- a/Engine/Plugins/Media/ImgMedia/Source/ImgMedia/Private/Readers/ExrImgMediaReader.h
+++ b/Engine/Plugins/Media/ImgMedia/Source/ImgMedia/Private/Readers/ExrImgMediaReader.h
@@ -97,7 +97,7 @@ protected:
 	 * @param Dim Dimensions of the largest mip.
 	 * @return Total size of all mip levels.
 	 */
-	SIZE_T GetMipBufferTotalSize(FIntPoint Dim);
+	SIZE_T GetMipBufferTotalSize(FIntPoint Dim, bool bInHasMips);
 
 protected:
 	TSet<int32> CanceledFrames;
diff --git a/Engine/Plugins/Media/ImgMedia/Source/ImgMedia/Private/Readers/IImgMediaReader.h b/Engine/Plugins/Media/ImgMedia/Source/ImgMedia/Private/Readers/IImgMediaReader.h
--- a/Engine/Plugins/Media/ImgMedia/Source/ImgMedia/Private/Readers/IImgMediaReader.h
+++ b/Engine/Plugins/Media/ImgMedia/Source/ImgMedia/Private/Readers/IImgMediaReader.h
@@ -31,7 +31,7 @@ struct FImgMediaFrameInfo
 	/** Whether the frame is in sRGB color space. */
 	bool Srgb;
 
-	/** Uncompressed size (in bytes). */
+	/** Uncompressed size (in bytes). All mip levels included. */
 	SIZE_T UncompressedSize;
 
 	/** Number of channels (RGB - 3 or RGBA - 4). */
diff --git a/Engine/Plugins/Media/ImgMedia/Source/OpenExrWrapper/Private/OpenExrWrapperContextReader.cpp b/Engine/Plugins/Media/ImgMedia/Source/OpenExrWrapper/Private/OpenExrWrapperContextReader.cpp
--- a/Engine/Plugins/Media/ImgMedia/Source/OpenExrWrapper/Private/OpenExrWrapperContextReader.cpp
+++ b/Engine/Plugins/Media/ImgMedia/Source/OpenExrWrapper/Private/OpenExrWrapperContextReader.cpp
@@ -123,11 +123,10 @@ FFrameRate FOpenExrHeaderReader::GetFrameRate(const FFrameRate& DefaultValue) co
 
 int32 FOpenExrHeaderReader::GetUncompressedSize() const
 {
-	const int32 NumChannels = GetNumChannels();
-	const int32 ChannelSize = sizeof(int16);
+	const int32 TexelSize = GetPixelSize();
 	const FIntPoint Window = GetDataWindow();
 
-	return (Window.X * Window.Y * NumChannels * ChannelSize);
+	return (Window.X * Window.Y * TexelSize);
 }
 
 int32 FOpenExrHeaderReader::GetNumChannels() const
@@ -137,6 +136,33 @@ int32 FOpenExrHeaderReader::GetNumChannels() const
 	return Attribute->chlist->num_channels;
 }
 
+int32 FOpenExrHeaderReader::GetPixelSize() const
+{
+	const exr_attribute_t* Attribute = nullptr;
+
+	CheckExrResult(exr_get_attribute_by_name(*((exr_context_t*)FileContext.Get()), 0, EXR_ATTRIBUTE_CHANNELS, &Attribute));
+
+	int32 PixelSizeInBytes = 0;
+	for (int ChannelId = 0; ChannelId < Attribute->chlist->num_channels; ChannelId++)
+	{
+		const exr_pixel_type_t ChannelType = Attribute->chlist->entries[ChannelId].pixel_type;
+		switch (ChannelType)
+		{
+		// 32 bit, 4 byte
+		case EXR_PIXEL_UINT:
+		case EXR_PIXEL_FLOAT:
+			PixelSizeInBytes += 4;
+			break;
+		// 16 bit, 2 byte
+		case EXR_PIXEL_HALF:
+		default:
+			PixelSizeInBytes += 2;
+			break;
+		}
+	}
+	return PixelSizeInBytes;
+}
+
 bool FOpenExrHeaderReader::ContainsMips() const
 {
 	const exr_attribute_t* Attribute = nullptr;
diff --git a/Engine/Plugins/Media/ImgMedia/Source/OpenExrWrapper/Public/OpenExrWrapper.h b/Engine/Plugins/Media/ImgMedia/Source/OpenExrWrapper/Public/OpenExrWrapper.h
--- a/Engine/Plugins/Media/ImgMedia/Source/OpenExrWrapper/Public/OpenExrWrapper.h
+++ b/Engine/Plugins/Media/ImgMedia/Source/OpenExrWrapper/Public/OpenExrWrapper.h
@@ -72,6 +72,7 @@ public:
 	FFrameRate GetFrameRate(const FFrameRate& DefaultValue) const;
 	int32 GetUncompressedSize() const;
 	int32 GetNumChannels() const;
+	int32 GetPixelSize() const;
 	bool ContainsMips() const;
 	int32 CalculateNumMipLevels(const FIntPoint& NumTiles) const;
 
