Fix warning V1052: Declaring virtual methods in a class marked as 'final' is pointless.
#rnx

diff --git a/Engine/Plugins/Media/ImgMedia/Source/ImgMedia/Private/Readers/IImgMediaReader.h b/Engine/Plugins/Media/ImgMedia/Source/ImgMedia/Private/Readers/IImgMediaReader.h
--- a/Engine/Plugins/Media/ImgMedia/Source/ImgMedia/Private/Readers/IImgMediaReader.h
+++ b/Engine/Plugins/Media/ImgMedia/Source/ImgMedia/Private/Readers/IImgMediaReader.h
@@ -333,7 +333,7 @@ struct FImgMediaFrame final
 	std::atomic<uint32> Stride = 0;
 
 	// This should only be used if you do not need to make changes to sample converter.
-	virtual IMediaTextureSampleConverter* GetSampleConverter()
+	IMediaTextureSampleConverter* GetSampleConverter()
 	{
 		return SampleConverter.Get();
 	};
