Img Media: Added a check to make sure that Exr Reader GPU can process the uncompressed exr.

diff --git a/Engine/Plugins/Media/ImgMedia/Source/ImgMedia/Private/Readers/ExrImgMediaReader.cpp b/Engine/Plugins/Media/ImgMedia/Source/ImgMedia/Private/Readers/ExrImgMediaReader.cpp
--- a/Engine/Plugins/Media/ImgMedia/Source/ImgMedia/Private/Readers/ExrImgMediaReader.cpp
+++ b/Engine/Plugins/Media/ImgMedia/Source/ImgMedia/Private/Readers/ExrImgMediaReader.cpp
@@ -344,6 +344,7 @@ void FExrImgMediaReader::UncancelFrame(int32 FrameNumber)
 TSharedPtr<IImgMediaReader, ESPMode::ThreadSafe> FExrImgMediaReader::GetReader(const TSharedRef <FImgMediaLoader, ESPMode::ThreadSafe>& InLoader, FString FirstImageInSequencePath)
 {
 	bool bIsCustomFormat = false;
+	bool bIsOptimizedForGpu = false;
 	FIntPoint TileSize(EForceInit::ForceInitToZero);
 
 #if defined(PLATFORM_WINDOWS) && PLATFORM_WINDOWS
@@ -359,6 +360,7 @@ TSharedPtr<IImgMediaReader, ESPMode::ThreadSafe> FExrImgMediaReader::GetReader(c
 	}
 
 	bIsCustomFormat = Info.FormatName.Equals(TEXT("EXR CUSTOM"));
+	bIsOptimizedForGpu = Info.FormatName.Equals(TEXT("EXR GPU")) || bIsCustomFormat;
 	if (bIsCustomFormat)
 	{
 		TileSize = Info.TileDimensions;
@@ -369,6 +371,7 @@ TSharedPtr<IImgMediaReader, ESPMode::ThreadSafe> FExrImgMediaReader::GetReader(c
 	if (GDynamicRHI && GDynamicRHI->GetInterfaceType() == ERHIInterfaceType::D3D12
 		&& Info.CompressionName == "Uncompressed" 
 		&& CVarEnableUncompressedExrGpuReader.GetValueOnAnyThread()
+		&& bIsOptimizedForGpu
 		)
 	{
 		TSharedRef<FExrImgMediaReaderGpu, ESPMode::ThreadSafe> GpuReader = 
@@ -416,7 +419,16 @@ bool FExrImgMediaReader::GetInfo(const FString& FilePath, FImgMediaFrameInfo& Ou
 	}
 	else
 	{
-		OutInfo.FormatName = TEXT("EXR");
+		// Can GPU reader be utilized to read this EXR file.
+		if (HeaderReader.IsOptimizedForGpu())
+		{
+			OutInfo.FormatName = TEXT("EXR GPU");
+		}
+		else
+		{
+			OutInfo.FormatName = TEXT("EXR");
+		}
+
 		OutInfo.bHasTiles = HeaderReader.GetTileSize(OutInfo.TileDimensions);
 		OutInfo.TileBorder = 0;
 	}
diff --git a/Engine/Plugins/Media/ImgMedia/Source/OpenExrWrapper/Private/OpenExrWrapperContextReader.cpp b/Engine/Plugins/Media/ImgMedia/Source/OpenExrWrapper/Private/OpenExrWrapperContextReader.cpp
--- a/Engine/Plugins/Media/ImgMedia/Source/OpenExrWrapper/Private/OpenExrWrapperContextReader.cpp
+++ b/Engine/Plugins/Media/ImgMedia/Source/OpenExrWrapper/Private/OpenExrWrapperContextReader.cpp
@@ -189,6 +189,23 @@ int32 FOpenExrHeaderReader::CalculateNumMipLevels(const FIntPoint& NumTiles) con
 	return 1;
 }
 
+bool FOpenExrHeaderReader::IsOptimizedForGpu() const
+{
+	const exr_attribute_t* Attribute = nullptr;
+
+	CheckExrResult(exr_get_attribute_by_name(*((exr_context_t*)FileContext.Get()), 0, EXR_ATTRIBUTE_CHANNELS, &Attribute));
+
+	for (int ChannelId = 0; ChannelId < Attribute->chlist->num_channels; ChannelId++)
+	{
+		const exr_pixel_type_t ChannelType = Attribute->chlist->entries[ChannelId].pixel_type;
+		if (ChannelType != EXR_PIXEL_HALF)
+		{
+			return false;
+		}
+	}
+	return Attribute->chlist->num_channels <= 4;
+}
+
 bool FOpenExrHeaderReader::GetTileSize(FIntPoint& OutTileSize) const
 {
 	const exr_attribute_t* Attribute = nullptr;
diff --git a/Engine/Plugins/Media/ImgMedia/Source/OpenExrWrapper/Public/OpenExrWrapper.h b/Engine/Plugins/Media/ImgMedia/Source/OpenExrWrapper/Public/OpenExrWrapper.h
--- a/Engine/Plugins/Media/ImgMedia/Source/OpenExrWrapper/Public/OpenExrWrapper.h
+++ b/Engine/Plugins/Media/ImgMedia/Source/OpenExrWrapper/Public/OpenExrWrapper.h
@@ -76,6 +76,9 @@ public:
 	bool ContainsMips() const;
 	int32 CalculateNumMipLevels(const FIntPoint& NumTiles) const;
 
+	/** Determines if Gpu reader can be utilized. */
+	bool IsOptimizedForGpu() const;
+
 	/** Gets tile dimensions. Returns false if image has no tiles. */
 	bool GetTileSize(FIntPoint& OutTileSize) const;
 
