Fix for blocked / synced playback timeouts and related issues with ImgMediaPlayer
#tests Various local tests with testbed project

diff --git a/Engine/Plugins/Media/ImgMedia/Source/ImgMedia/Internal/Loader/ImgMediaLoader.h b/Engine/Plugins/Media/ImgMedia/Source/ImgMedia/Internal/Loader/ImgMediaLoader.h
--- a/Engine/Plugins/Media/ImgMedia/Source/ImgMedia/Internal/Loader/ImgMediaLoader.h
+++ b/Engine/Plugins/Media/ImgMedia/Source/ImgMedia/Internal/Loader/ImgMediaLoader.h
@@ -369,9 +369,9 @@ public:
 	bool RequestFrame(FTimespan Time, float PlayRate, bool Loop);
 
 	/**
-	 * Reset "queued fetch" related state used to emulate player output queue behavior
+	 * Flush state if queue emulation
 	 */
-	void ResetFetchLogic();
+	void Flush();
 
 	/**
 	 * Tell the loader if playback is blocking.
@@ -547,6 +547,11 @@ protected:
 	 */
 	const TSharedPtr<FImgMediaFrame, ESPMode::ThreadSafe>* GetCachedFrame(int32 InFrameNumber);
 
+	/**
+	 * Reset "queued fetch" related state used to emulate player output queue behavior
+	 */
+	void ResetFetchLogic();
+
 private:
 
 	/** Critical section for synchronizing access to Frames. */
@@ -638,10 +643,9 @@ private:
 	/** State related to "queue style" frame access functions */
 	struct
 	{
-		int32 LastFrameIndex;
-		FMediaTimeStamp LastTimeStamp;
-		FTimespan LastDuration;
-		int32 LoopIndex;
+		FMediaTimeStamp LastTimeStamp;	// Current playhead / last delivered sample PTS
+		FTimespan LastDuration;			// Duration of a frame / zero after seek or playback start
+		int32 LoopIndex;				// Current loop index to be encoded into sequence index of timestamp
 	} QueuedSampleFetch;
 
 private:
diff --git a/Engine/Plugins/Media/ImgMedia/Source/ImgMedia/Private/Loader/ImgMediaLoader.cpp b/Engine/Plugins/Media/ImgMedia/Source/ImgMedia/Private/Loader/ImgMediaLoader.cpp
--- a/Engine/Plugins/Media/ImgMedia/Source/ImgMedia/Private/Loader/ImgMediaLoader.cpp
+++ b/Engine/Plugins/Media/ImgMedia/Source/ImgMedia/Private/Loader/ImgMediaLoader.cpp
@@ -186,6 +186,10 @@ FImgMediaLoader::FImgMediaLoader(const TSharedRef<FImgMediaScheduler, ESPMode::T
 #endif
 {
 	ResetFetchLogic();
+	QueuedSampleFetch.LastTimeStamp.Time = FTimespan::Zero();
+	QueuedSampleFetch.LastTimeStamp.SequenceIndex = 0;
+	QueuedSampleFetch.LastDuration = FTimespan::Zero();
+
 	UE_LOG(LogImgMedia, Verbose, TEXT("Loader %p: Created"), this);
 }
 
@@ -300,9 +304,6 @@ const FString& FImgMediaLoader::GetImagePath(int32 FrameNumber, int32 MipLevel)
 
 void FImgMediaLoader::ResetFetchLogic()
 {
-	QueuedSampleFetch.LastFrameIndex = INDEX_NONE;
-	QueuedSampleFetch.LastTimeStamp.Invalidate();
-	QueuedSampleFetch.LastDuration = FTimespan::Zero();
 	QueuedSampleFetch.LoopIndex = 0;
 
 	SkipFramesCounter = 0;
@@ -322,18 +323,18 @@ void FImgMediaLoader::HandlePause()
 #endif
 }
 
+void FImgMediaLoader::Flush()
+{
+	ResetFetchLogic();
+	// Adjust last timestamp / play head so it is within the current, reset sequence index range
+	QueuedSampleFetch.LastTimeStamp.SequenceIndex = FMediaTimeStamp::MakeSequenceIndex(0, 0);
+}
 
 void FImgMediaLoader::Seek(const FMediaTimeStamp& SeekTarget, bool bReverse)
 {
 	ResetFetchLogic();
-	if (!bReverse)
-	{
-		QueuedSampleFetch.LastTimeStamp = SeekTarget;
-	}
-	else
-	{
-		QueuedSampleFetch.LastTimeStamp = SeekTarget - FTimespan::FromSeconds(static_cast<double>(SequenceFrameRate.Denominator) / SequenceFrameRate.Numerator);
-	}
+	QueuedSampleFetch.LastTimeStamp = { FrameNumberToTime(TimeToFrameNumber(SeekTarget.Time)), SeekTarget.SequenceIndex };
+	QueuedSampleFetch.LastDuration = FTimespan::Zero();
 }
 
 
@@ -341,31 +342,30 @@ bool FImgMediaLoader::DiscardVideoSamples(const TRange<FMediaTimeStamp>& TimeRan
 {
 	bool bOk = false;
 
-	if (QueuedSampleFetch.LastTimeStamp.IsValid())
+	check(QueuedSampleFetch.LastTimeStamp.IsValid());
+
+	if (PlayRate >= 0.0f)
 	{
-		if (PlayRate >= 0.0f)
+		if (QueuedSampleFetch.LastTimeStamp < TimeRange.GetUpperBoundValue())
 		{
-			if ((QueuedSampleFetch.LastTimeStamp + QueuedSampleFetch.LastDuration) < TimeRange.GetUpperBoundValue())
-			{
-				QueuedSampleFetch.LastTimeStamp = TimeRange.GetUpperBoundValue();
-				bOk = true;
-			}
+			QueuedSampleFetch.LastTimeStamp = TimeRange.GetUpperBoundValue();
+			bOk = true;
 		}
-		else
+	}
+	else
+	{
+		if (QueuedSampleFetch.LastTimeStamp >= TimeRange.GetLowerBoundValue())
 		{
-			if ((QueuedSampleFetch.LastTimeStamp - QueuedSampleFetch.LastDuration) > TimeRange.GetLowerBoundValue())
-			{
-				QueuedSampleFetch.LastTimeStamp = TimeRange.GetLowerBoundValue();
-				bOk = true;
-			}
+			QueuedSampleFetch.LastTimeStamp = TimeRange.GetLowerBoundValue() - FTimespan::FromSeconds(SequenceFrameRate.AsInterval());
+			bOk = true;
 		}
+	}
 
-		if (bOk)
-		{
-			QueuedSampleFetch.LastDuration = FTimespan::Zero();
-			QueuedSampleFetch.LastTimeStamp.Time = FrameNumberToTime(TimeToFrameNumber(QueuedSampleFetch.LastTimeStamp.Time));
-			QueuedSampleFetch.LastFrameIndex = INDEX_NONE;
-		}
+	if (bOk)
+	{
+		ResetFetchLogic();
+		QueuedSampleFetch.LastTimeStamp.Time = FrameNumberToTime(TimeToFrameNumber(QueuedSampleFetch.LastTimeStamp.Time));
+		QueuedSampleFetch.LastDuration = FTimespan::Zero();
 	}
 
 	return bOk;
@@ -379,37 +379,35 @@ IMediaSamples::EFetchBestSampleResult FImgMediaLoader::FetchBestVideoSampleForTi
 		// A few sanity checks
 		check(InTimeRange.GetLowerBoundValue() <= InTimeRange.GetUpperBoundValue());
 		check(FMediaTimeStamp::GetPrimaryIndex(InTimeRange.GetLowerBoundValue().SequenceIndex) == FMediaTimeStamp::GetPrimaryIndex(InTimeRange.GetUpperBoundValue().SequenceIndex));
+		check(QueuedSampleFetch.LastTimeStamp.IsValid());
 
 		auto TimeRange = InTimeRange;
 
 		// Clamp the current range requested against the last known timestamp to ensure we return the "next" sample at all times...
-		if (QueuedSampleFetch.LastTimeStamp.IsValid())
+		if (PlayRate >= 0.0f)
 		{
-			if (PlayRate >= 0.0f)
+			FMediaTimeStamp NextTimeStamp = QueuedSampleFetch.LastTimeStamp + QueuedSampleFetch.LastDuration;
+			if (NextTimeStamp >= TimeRange.GetUpperBoundValue())
 			{
-				FMediaTimeStamp NextTimeStamp = QueuedSampleFetch.LastTimeStamp + QueuedSampleFetch.LastDuration;
-				if (NextTimeStamp >= TimeRange.GetUpperBoundValue())
-					{
-					return IMediaSamples::EFetchBestSampleResult::NoSample;
-					}
+				return IMediaSamples::EFetchBestSampleResult::NoSample;
+			}
 
-				if (TimeRange.GetLowerBoundValue() < NextTimeStamp)
-				{
-					TimeRange = TRange<FMediaTimeStamp>(NextTimeStamp, InTimeRange.GetUpperBoundValue());
-				}
+			if (TimeRange.GetLowerBoundValue() < NextTimeStamp)
+			{
+				TimeRange = TRange<FMediaTimeStamp>(NextTimeStamp, InTimeRange.GetUpperBoundValue());
 			}
-			else
+		}
+		else
+		{
+			FMediaTimeStamp NextTimeStampAtEnd = (QueuedSampleFetch.LastTimeStamp - QueuedSampleFetch.LastDuration) + FTimespan::FromSeconds(SequenceFrameRate.AsInterval());
+			if (NextTimeStampAtEnd < TimeRange.GetLowerBoundValue())
 			{
-				FMediaTimeStamp NextTimeStamp = QueuedSampleFetch.LastTimeStamp;
-				if (NextTimeStamp <= TimeRange.GetLowerBoundValue())
-				{
-					return IMediaSamples::EFetchBestSampleResult::NoSample;
-				}
+				return IMediaSamples::EFetchBestSampleResult::NoSample;
+			}
 
-				if (TimeRange.GetUpperBoundValue() > NextTimeStamp)
-				{
-					TimeRange = TRange<FMediaTimeStamp>(InTimeRange.GetLowerBoundValue(), NextTimeStamp);
-				}
+			if (TimeRange.GetUpperBoundValue() > NextTimeStampAtEnd)
+			{
+				TimeRange = TRange<FMediaTimeStamp>(InTimeRange.GetLowerBoundValue(), NextTimeStampAtEnd);
 			}
 		}
 
@@ -417,6 +415,7 @@ IMediaSamples::EFetchBestSampleResult FImgMediaLoader::FetchBestVideoSampleForTi
 		int32 StartLoopIndex = FMediaTimeStamp::GetSecondaryIndex(TimeRange.GetLowerBoundValue().SequenceIndex);
 		FTimespan EndTime = TimeRange.GetUpperBoundValue().Time;
 		int32 EndLoopIndex;
+
 		// The end time is EXCLUSIVE, move the time by the smallest amount we can so we correctly select indices of frames
 		if (EndTime > FTimespan::Zero())
 		{
@@ -468,7 +467,7 @@ IMediaSamples::EFetchBestSampleResult FImgMediaLoader::FetchBestVideoSampleForTi
 			// More than one possibility?
 			if (EndIndex != StartIndex || StartLoopIndex != EndLoopIndex)
 			{
-				FTimespan FrameLength = FTimespan::FromSeconds(static_cast<double>(SequenceFrameRate.Denominator) / SequenceFrameRate.Numerator);
+				FTimespan FrameLength = FTimespan::FromSeconds(SequenceFrameRate.AsInterval());
 
 				// Compute visible duration of first frame...
 				FTimespan FirstTime = FrameNumberToTime(StartIndex);
@@ -509,7 +508,7 @@ IMediaSamples::EFetchBestSampleResult FImgMediaLoader::FetchBestVideoSampleForTi
 			// More than one possibility?
 			if (EndIndex != StartIndex || StartLoopIndex != EndLoopIndex)
 			{
-				FTimespan FrameLength = FTimespan::FromSeconds(static_cast<double>(SequenceFrameRate.Denominator) / SequenceFrameRate.Numerator);
+				FTimespan FrameLength = FTimespan::FromSeconds(SequenceFrameRate.AsInterval());
 
 				// Compute the length of the last frame
 				FTimespan FirstTime = FrameNumberToTime(EndIndex);
@@ -519,8 +518,8 @@ IMediaSamples::EFetchBestSampleResult FImgMediaLoader::FetchBestVideoSampleForTi
 				FTimespan LastTime = FrameNumberToTime(StartIndex);
 				FTimespan LastFrameDuration = FrameLength - FTimespan(FMath::Max((TimeRange.GetLowerBoundValue().Time - LastTime).GetTicks(), 0));
 
-				// First one wins by default
-				FTimespan MaxFrameDuration = FirstFrameDuration;
+				// Last one wins by default
+				FTimespan MaxFrameDuration = LastFrameDuration;
 
 				int32 NumFrames = GetNumImages();
 
@@ -533,20 +532,22 @@ IMediaSamples::EFetchBestSampleResult FImgMediaLoader::FetchBestVideoSampleForTi
 					MaxFrameDuration = FrameLength;
 				}
 
-				// Is the last one even better?
-				if (LastFrameDuration >= MaxFrameDuration)
+				// Is the first one even better?
+				if (FirstFrameDuration >= MaxFrameDuration)
 				{
 					MaxIdx = StartIndex;
 				}
 			}
 		}
 
+		// - - - - - - - - - - - - - - - - - - - - - - - - - - - - - - - - - - - - - - - - - - - - - - - - - - - -
+
 		// Anything?
 		if (MaxIdx >= 0)
 		{
 			const TSharedPtr<FImgMediaFrame, ESPMode::ThreadSafe>* Frame;
 
-			// Request data for the frame we would like... (in case it's not in, yet)
+			// Request data for the frame we would like... (in case it's not requested or in, yet)
 			RequestFrame(MaxIdx, PlayRate, bIsLoopingEnabled);
 
 			FScopeLock Lock(&CriticalSection);
@@ -571,7 +572,7 @@ IMediaSamples::EFetchBestSampleResult FImgMediaLoader::FetchBestVideoSampleForTi
 
 				// Check if we wrapped around the end of the media vs. the last frame we returned...
 				// (skip this if we are not looping or if this is either the first frame ever or first frame after a seek)
-				if (bIsLoopingEnabled && QueuedSampleFetch.LastTimeStamp.IsValid() && QueuedSampleFetch.LastFrameIndex != INDEX_NONE)
+				if (bIsLoopingEnabled && !QueuedSampleFetch.LastDuration.IsZero())
 				{
 					if (PlayRate >= 0.0f)
 					{
@@ -590,14 +591,12 @@ IMediaSamples::EFetchBestSampleResult FImgMediaLoader::FetchBestVideoSampleForTi
 				}
 
 				// Track state & setup sample for caller...
-				QueuedSampleFetch.LastFrameIndex = MaxIdx;
 				QueuedSampleFetch.LastTimeStamp = FMediaTimeStamp(SampleTime, FMediaTimeStamp::MakeSequenceIndex(FMediaTimeStamp::GetPrimaryIndex(TimeRange.GetLowerBoundValue().SequenceIndex), QueuedSampleFetch.LoopIndex));
 				QueuedSampleFetch.LastDuration = FTimespan::FromSeconds(Frame->Get()->GetFrameRate().AsInterval());
 
 				// We are clear to return it as new result... Make a sample & initialize it...
 				auto Sample = MakeShared<FImgMediaTextureSample, ESPMode::ThreadSafe>();
 
-
 				ImgMediaLoader::CheckAndUpdateImgDimensions(SequenceDim, Frame->Get()->GetDim());
 
 				if (Sample->Initialize(*Frame->Get(), SequenceDim, QueuedSampleFetch.LastTimeStamp, QueuedSampleFetch.LastDuration, GetNumMipLevels(), TilingDescription))
@@ -641,108 +640,35 @@ IMediaSamples::EFetchBestSampleResult FImgMediaLoader::FetchBestVideoSampleForTi
 	return IMediaSamples::EFetchBestSampleResult::NoSample;
 }
 
-
+ 
 bool FImgMediaLoader::PeekVideoSampleTime(FMediaTimeStamp &TimeStamp, bool bIsLoopingEnabled, float PlayRate, const FTimespan& CurrentTime)
 {
+	// Strictly speaking we never actual return a "frame" here as we do NOT involve any requests to the cache -- we simply assume we could deliver what a "queue" would contain on demand later...
 	if (IsInitialized())
 	{
-		int32 Idx = INDEX_NONE;
-		bool bNewLoopSeq = false;
-		FTimespan FrameStart;
+		check(QueuedSampleFetch.LastTimeStamp.IsValid());
 
-		// Do we know which index we handed out last?
-		if (QueuedSampleFetch.LastFrameIndex != INDEX_NONE)
+		if (PlayRate >= 0.0f)
 		{
-			// Yes. A queue would now yield the next frame (independent of rate). See which index that would be...
-			Idx = int32(QueuedSampleFetch.LastFrameIndex + FMath::Sign(PlayRate));
-			int32 NumFrames = GetNumImages();
-			if (bIsLoopingEnabled)
-			{
-				if (Idx < 0)
-				{
-					Idx = NumFrames - 1;
-					bNewLoopSeq = true;
-				}
-				else if (Idx >= NumFrames)
-				{
-					Idx = 0;
-					bNewLoopSeq = true;
-				}
-			}
-			else
+			TimeStamp = QueuedSampleFetch.LastTimeStamp + QueuedSampleFetch.LastDuration;
+			if (TimeStamp >= SequenceDuration)
 			{
-				// If we reach either end of the sequence with no looping, we have no more frames to offer
-				if (Idx < 0 || Idx >= NumFrames)
-				{
-					return false;
-				}
+				TimeStamp -= SequenceDuration;
+				TimeStamp.SequenceIndex = FMediaTimeStamp::AdjustSecondaryIndex(QueuedSampleFetch.LastTimeStamp.SequenceIndex, 1);
 			}
-
-			FrameStart = FrameNumberToTime(Idx);
 		}
 		else
 		{
-			if (QueuedSampleFetch.LastTimeStamp.IsValid())
+			TimeStamp = QueuedSampleFetch.LastTimeStamp - QueuedSampleFetch.LastDuration;
+			if (TimeStamp.Time.GetTicks() < 0)
 			{
-				if (PlayRate >= 0.0f)
-				{
-					auto NextTime = QueuedSampleFetch.LastTimeStamp.Time + QueuedSampleFetch.LastDuration;
-					if (NextTime >= SequenceDuration)
-					{
-						if (!bIsLoopingEnabled)
-						{
-							return false;
-						}
-						NextTime -= SequenceDuration;
-						bNewLoopSeq = true;
-					}
-					FrameStart = NextTime;
-				}
-				else
-				{
-					auto NextTime = QueuedSampleFetch.LastTimeStamp.Time - QueuedSampleFetch.LastDuration;
-					if (NextTime < FTimespan::Zero())
-					{
-						if (!bIsLoopingEnabled)
-						{
-							return false;
-						}
-						NextTime += SequenceDuration;
-						bNewLoopSeq = true;
-					}
-					FrameStart = NextTime;
-				}
-			}
-			else
-			{
-				FrameStart = CurrentTime;
+				TimeStamp += SequenceDuration;
+				TimeStamp.SequenceIndex = FMediaTimeStamp::AdjustSecondaryIndex(QueuedSampleFetch.LastTimeStamp.SequenceIndex, -1);
 			}
 		}
 
-		if (Idx == INDEX_NONE)
-		{
-			Idx = TimeToFrameNumber(FrameStart);
-		}
-
-		auto SequenceIndex = QueuedSampleFetch.LastTimeStamp.IsValid() ? QueuedSampleFetch.LastTimeStamp.SequenceIndex : 0;
-
-		// If possible, fetch any existing frame data...
-		// (just to see if we have any)
-		const TSharedPtr<FImgMediaFrame, ESPMode::ThreadSafe>* Frame = ((uint32)Idx != INDEX_NONE) ? GetCachedFrame(Idx) : nullptr;
-
-		// Data is present?
-		if (Frame)
-		{
-			// Yes. Return the timing information...
-			TimeStamp.Time = FrameStart;
-			TimeStamp.SequenceIndex = bNewLoopSeq ? FMediaTimeStamp::AdjustSecondaryIndex(SequenceIndex, (PlayRate >= 0.0f) ? 1 : -1) : SequenceIndex;
-			return true;
-		}
-		else
-		{
-			// No data. We will request it (so, like other players do) we fill our (virtual) queue at the current location automatically
-			RequestFrame(Idx, PlayRate, bIsLoopingEnabled);
-		}
+		TimeStamp.Time = FrameNumberToTime(TimeToFrameNumber(TimeStamp.Time));
+		return true;
 	}
 	return false;
 }
diff --git a/Engine/Plugins/Media/ImgMedia/Source/ImgMedia/Private/Player/ImgMediaPlayer.cpp b/Engine/Plugins/Media/ImgMedia/Source/ImgMedia/Private/Player/ImgMediaPlayer.cpp
--- a/Engine/Plugins/Media/ImgMedia/Source/ImgMedia/Private/Player/ImgMediaPlayer.cpp
+++ b/Engine/Plugins/Media/ImgMedia/Source/ImgMedia/Private/Player/ImgMediaPlayer.cpp
@@ -650,7 +650,7 @@ void FImgMediaPlayer::FlushSamples()
 	{
 		CurrentSeekIndex = 0;
 		LastFetchTime = FTimespan::MinValue();
-		Loader->ResetFetchLogic();
+		Loader->Flush();
 	}
 }
 
diff --git a/Engine/Source/Runtime/MediaUtils/Private/MediaPlayerFacade.cpp b/Engine/Source/Runtime/MediaUtils/Private/MediaPlayerFacade.cpp
--- a/Engine/Source/Runtime/MediaUtils/Private/MediaPlayerFacade.cpp
+++ b/Engine/Source/Runtime/MediaUtils/Private/MediaPlayerFacade.cpp
@@ -1667,9 +1667,12 @@ void FMediaPlayerFacade::Flush(bool bExcludePlayer, bool bOnSeek)
 		SeekIndex = 0;
 	}
 
-	// Logically we have no old sample anymore
+	// Logically we have no old sample anymore if we did seek
 	// (as in: we will start asking for a new one until we get one - even with a rate of zero, if we had a non-zero one ever before)
-	LastVideoSampleProcessedTimeRange = TRange<FMediaTimeStamp>::Empty();
+	if (bOnSeek)
+	{
+		LastVideoSampleProcessedTimeRange = TRange<FMediaTimeStamp>::Empty();
+	}
 
 	// Invalidate next video time to fetch (none-audio case)
 	NextEstVideoTimeAtFrameStart.Invalidate();
@@ -2896,45 +2899,68 @@ void FMediaPlayerFacade::ProcessAudioSamples(IMediaSamples& Samples, const TRang
 bool FMediaPlayerFacade::IsVideoSampleStillGood(const TRange<FMediaTimeStamp>& LastSampleTimeRange, const TRange<FMediaTimeStamp>& TimeRange, bool bReverse) const
 {
 	// If we have no valid time range or a seek is in progress we assume the current frame can be considered "done" in any case
-	if (!TimeRange.IsEmpty() && !SeekTargetTime.IsValid())
+	if (!TimeRange.IsEmpty() && !SeekTargetTime.IsValid() && !LastSampleTimeRange.IsEmpty())
 	{
 		// This is not the case: check more detailed!
-		TRange<FMediaTimeStamp> LastSampleTimeRange0(FMediaTimeStamp(LastSampleTimeRange.GetLowerBoundValue().Time, 0), FMediaTimeStamp(LastSampleTimeRange.GetUpperBoundValue().Time, 0));
-		TRange<FMediaTimeStamp> TimeRange0(FMediaTimeStamp(TimeRange.GetLowerBoundValue().Time, 0), FMediaTimeStamp(TimeRange.GetUpperBoundValue().Time, 0));
 
-		// Is the sample time range at all still valid?
-		if (LastSampleTimeRange0.Overlaps(TimeRange0))
+		// This better be true at all times
+		check(LastSampleTimeRange.GetLowerBoundValue().SequenceIndex == LastSampleTimeRange.GetUpperBoundValue().SequenceIndex);
+
+		TRange<FMediaTimeStamp> TimeRange0;
+
+		// If we encounter a time range crossing some sequence index change, we need to check if we can "unroll" it...
+		uint64 LowerSeqIdx = TimeRange.GetLowerBoundValue().SequenceIndex;
+		uint64 UpperSeqIdx = TimeRange.GetUpperBoundValue().SequenceIndex;
+		if (LowerSeqIdx != UpperSeqIdx)
 		{
-			// Yes. Assuming we could get more samples (of the same type) from the player, would the next one "better"?
-			// (this does not say which "next one" is actually current (e.g. if the rendering FPS is real bad), just that "another one" would be better)
+			if (FMediaTimeStamp::GetPrimaryIndex(LowerSeqIdx) != FMediaTimeStamp::GetPrimaryIndex(UpperSeqIdx))
+			{
+				// If we have a primary index change, we cannot assume any valid frame around...
+				return false;
+			}
 
-			// Compute the "theoretical" next sample range...
-			TRange<FMediaTimeStamp> NextSampleTimeRange = !bReverse ? TRange<FMediaTimeStamp>(LastSampleTimeRange0.GetUpperBoundValue(), LastSampleTimeRange0.GetUpperBoundValue() + LastSampleTimeRange0.Size<FMediaTimeStamp>().Time)
-				: TRange<FMediaTimeStamp>(LastSampleTimeRange0.GetLowerBoundValue() - LastSampleTimeRange0.Size<FMediaTimeStamp>().Time, LastSampleTimeRange0.GetLowerBoundValue());
+			// So we must have a loop index change. Compute how many loops and change the range into one "unrolled" one as indicated by the playback direction...
+			int64 LoopIdxDiff = FMediaTimeStamp::GetSecondaryIndex(LowerSeqIdx)
+							  - FMediaTimeStamp::GetSecondaryIndex(UpperSeqIdx);
+			check(LoopIdxDiff > 0);
 
-			FTimespan Duration = Player->GetControls().GetDuration();
+			double Duration = Player->GetControls().GetDuration().GetTotalSeconds();
 
-			if (!Player->GetControls().IsLooping())
+			if (!bReverse)
 			{
-				// If we are not looping we need to clamp against the media's duration
-				// (we assume it starts at zero here!)
-				NextSampleTimeRange = TRange<FMediaTimeStamp>::Intersection(NextSampleTimeRange, TRange<FMediaTimeStamp>(FMediaTimeStamp(0, NextSampleTimeRange.GetLowerBoundValue().SequenceIndex), FMediaTimeStamp(Duration, NextSampleTimeRange.GetLowerBoundValue().SequenceIndex)));
+				TimeRange0 = TRange<FMediaTimeStamp>(FMediaTimeStamp(TimeRange.GetLowerBoundValue().Time, 0), FMediaTimeStamp(TimeRange.GetUpperBoundValue().Time + FTimespan::FromSeconds(LoopIdxDiff * Duration), 0));
 			}
 			else
 			{
-				if (NextSampleTimeRange.GetLowerBoundValue().Time >= Duration)
-				{
-					check(!bReverse);
-					NextSampleTimeRange = TRange<FMediaTimeStamp>(FMediaTimeStamp(NextSampleTimeRange.GetLowerBoundValue().Time - Duration, FMediaTimeStamp::AdjustPrimaryIndex(NextSampleTimeRange.GetLowerBoundValue().SequenceIndex, 1)),
-						FMediaTimeStamp(NextSampleTimeRange.GetUpperBoundValue().Time - Duration, FMediaTimeStamp::AdjustPrimaryIndex(NextSampleTimeRange.GetUpperBoundValue().SequenceIndex, 1)));
-				}
-				else if (NextSampleTimeRange.GetLowerBoundValue().Time < FTimespan::Zero())
-				{
-					check(bReverse);
-					NextSampleTimeRange = TRange<FMediaTimeStamp>(FMediaTimeStamp(NextSampleTimeRange.GetLowerBoundValue().Time + Duration, FMediaTimeStamp::AdjustPrimaryIndex(NextSampleTimeRange.GetLowerBoundValue().SequenceIndex, -1)),
-						FMediaTimeStamp(NextSampleTimeRange.GetUpperBoundValue().Time + Duration, FMediaTimeStamp::AdjustPrimaryIndex(NextSampleTimeRange.GetUpperBoundValue().SequenceIndex, -1)));
-				}
+				TimeRange0 = TRange<FMediaTimeStamp>(FMediaTimeStamp(TimeRange.GetLowerBoundValue().Time - FTimespan::FromSeconds(LoopIdxDiff * Duration), 0), FMediaTimeStamp(TimeRange.GetUpperBoundValue().Time, 0));
 			}
+		}
+		else
+		{
+			// Simple case, just bring everything down to "zero sequence index" for ease of processing below...
+			TimeRange0 = TRange<FMediaTimeStamp>(FMediaTimeStamp(TimeRange.GetLowerBoundValue().Time, 0), FMediaTimeStamp(TimeRange.GetUpperBoundValue().Time, 0));
+		}
+
+		TRange<FMediaTimeStamp> LastSampleTimeRange0(FMediaTimeStamp(LastSampleTimeRange.GetLowerBoundValue().Time, 0), FMediaTimeStamp(LastSampleTimeRange.GetUpperBoundValue().Time, 0));
+
+		// Is the sample time range at all still valid?
+		if (LastSampleTimeRange0.Overlaps(TimeRange0))
+		{
+			// Yes. Assuming we could get more samples (of the same type) from the player, would the next one "better"?
+			// (this does not say which "next one" is actually current (e.g. if the rendering FPS is real bad), just that "another one" would be better)
+
+			// Compute the "theoretical" next sample range...
+			TRange<FMediaTimeStamp> NextSampleTimeRange = !bReverse ? TRange<FMediaTimeStamp>(LastSampleTimeRange0.GetUpperBoundValue(), LastSampleTimeRange0.GetUpperBoundValue() + LastSampleTimeRange0.Size<FMediaTimeStamp>().Time)
+																	: TRange<FMediaTimeStamp>(LastSampleTimeRange0.GetLowerBoundValue() - LastSampleTimeRange0.Size<FMediaTimeStamp>().Time, LastSampleTimeRange0.GetLowerBoundValue());
+
+			// Note: Loops (or the end of the timeline in non-looping setups)
+			// 
+			// - We could check for them and generate proper changes to the sequence index
+			// - Doing this would leave us with quite complex setups to compute the coverage
+			// - We opt for a cleaner, simpler approach: as we are NOT interested into proper PTS values, we can safely work with an "inifite" timeline when computing any overlaps, coverage and such
+			// 
+			// --> we simply keep what we compute above!
+			//
 
 			// Compute which one is larger inside the current range...
 			int64 LastSampleCoverage = TRange<FMediaTimeStamp>::Intersection(TimeRange0, LastSampleTimeRange0).Size<FMediaTimeStamp>().Time.GetTicks();
@@ -3033,7 +3059,8 @@ bool FMediaPlayerFacade::ProcessVideoSamples(IMediaSamples& Samples, const TRang
 		// Yes. If we are in blocking playback mode we need to make sure that the sample is really in the range we asked for and block on...
 		// (same players might return an older sample as stop-gap measure if nothing can be found in the current range)
 
-		TRange<FMediaTimeStamp> SampleTimeRange(Sample->GetTime(), Sample->GetTime() + Sample->GetDuration());
+		FMediaTimeStamp SampleTime = Sample->GetTime();
+		TRange<FMediaTimeStamp> SampleTimeRange(SampleTime, SampleTime + Sample->GetDuration());
 
 		// Enqueue the sample to render
 		// (we use a queue to stay compatible with existing structure and older sinks - new sinks will read this single entry right away on the gamethread
