Image Media Process Exr: Added option to remove alpha channel from exrs and removed old custom EXR code.

diff --git a/Engine/Plugins/Media/ImgMedia/Source/ImgMediaEditor/Private/Widgets/ImgMediaProcessEXROptions.h b/Engine/Plugins/Media/ImgMedia/Source/ImgMediaEditor/Private/Widgets/ImgMediaProcessEXROptions.h
--- a/Engine/Plugins/Media/ImgMedia/Source/ImgMediaEditor/Private/Widgets/ImgMediaProcessEXROptions.h
+++ b/Engine/Plugins/Media/ImgMedia/Source/ImgMediaEditor/Private/Widgets/ImgMediaProcessEXROptions.h
@@ -47,6 +47,12 @@ public:
 	UPROPERTY(EditAnywhere, Transient, Category = Processing)
 	bool bUsePlayer = false;
 
+	/** 
+	* This option removes alpha channel. Since alpha channel adds 25% reading cost, it is suggested to have alpha channel removed if it is not used. 
+	*/
+	UPROPERTY(EditAnywhere, Category = Channels)
+	bool bRemoveAlphaChannel = false;
+
 	/** Tint each mip level a different colour to help with debugging. */
 	UPROPERTY(EditAnywhere, Transient, Category = Debug)
 	bool bEnableMipLevelTint = false;
diff --git a/Engine/Plugins/Media/ImgMedia/Source/ImgMediaEditor/Private/Widgets/SImgMediaProcessEXR.cpp b/Engine/Plugins/Media/ImgMedia/Source/ImgMediaEditor/Private/Widgets/SImgMediaProcessEXR.cpp
--- a/Engine/Plugins/Media/ImgMedia/Source/ImgMediaEditor/Private/Widgets/SImgMediaProcessEXR.cpp
+++ b/Engine/Plugins/Media/ImgMedia/Source/ImgMediaEditor/Private/Widgets/SImgMediaProcessEXR.cpp
@@ -317,7 +317,7 @@ void SImgMediaProcessEXR::ProcessAllImages()
 					}
 
 					ProcessImageCustom(Image, InTileWidth, InTileHeight, TileBorder,
-						bEnableMips, ImageParameters, Name, false /*bUseCustomFormat*/);
+						bEnableMips, ImageParameters, Name);
 
 					NumActive--;
 				});
@@ -393,8 +393,7 @@ void SImgMediaProcessEXR::GetGlobalImageParameters()
 
 void SImgMediaProcessEXR::ProcessImageCustom(const FImage& InImage,
 	int32 InTileWidth, int32 InTileHeight, int32 InTileBorder, bool bInEnableMips,
-	const FImageParameters& ImageParameters, const FString& InName,
-	bool bIsCustomFormat)
+	const FImageParameters& ImageParameters, const FString& InName)
 {
 #if IMGMEDIAEDITOR_EXR_SUPPORTED_PLATFORM
 	TRACE_CPUPROFILER_EVENT_SCOPE(SImgMediaProcessEXR::ProcessImageCustom);
@@ -405,7 +404,7 @@ void SImgMediaProcessEXR::ProcessImageCustom(const FImage& InImage,
 
 	ProcessImageCustomRawData(RawData, Width, Height,
 		InTileWidth, InTileHeight, InTileBorder, bInEnableMips,
-		ImageParameters, InName, bIsCustomFormat);
+		ImageParameters, InName);
 #else // IMGMEDIAEDITOR_EXR_SUPPORTED_PLATFORM
 	UE_LOG(LogImgMediaEditor, Error, TEXT("EXR not supported on this platform."));
 #endif // IMGMEDIAEDITOR_EXR_SUPPORTED_PLATFORM
@@ -414,8 +413,7 @@ void SImgMediaProcessEXR::ProcessImageCustom(const FImage& InImage,
 void SImgMediaProcessEXR::ProcessImageCustomRawData(TArray64<uint8>& RawData,
 	int32 Width, int32 Height,
 	int32 InTileWidth, int32 InTileHeight, int32 InTileBorder, bool bInEnableMips,
-	const FImageParameters& ImageParameters, const FString& InName,
-	bool bIsCustomFormat)
+	const FImageParameters& ImageParameters, const FString& InName)
 {
 #if IMGMEDIAEDITOR_EXR_SUPPORTED_PLATFORM
 	TRACE_CPUPROFILER_EVENT_SCOPE(SImgMediaProcessEXR::ProcessImageCustomRawData);
@@ -431,7 +429,7 @@ void SImgMediaProcessEXR::ProcessImageCustomRawData(TArray64<uint8>& RawData,
 	int32 DestNumChannels = NumChannels;
 
 	// ImageWrapper always returns an alpha channel, so make sure we really have one.
-	if ((DestNumChannels == 4) && (ImageParameters.bHasAlphaChannel == false))
+	if ((DestNumChannels == 4) && ((ImageParameters.bHasAlphaChannel == false) || Options->bRemoveAlphaChannel))
 	{
 		// Remove the alpha channel as its not needed.
 		RemoveAlphaChannel(RawData);
@@ -466,21 +464,7 @@ void SImgMediaProcessEXR::ProcessImageCustomRawData(TArray64<uint8>& RawData,
 
 	// Create tiled exr file.
 	FTiledOutputFile OutFile(DisplayWindow.Min, DisplayWindow.Max,
-		DataWindow.Min, DataWindow.Max, bIsTiled || bIsCustomFormat);
-
-	// Add attributes.
-	if (bIsCustomFormat)
-	{
-		OutFile.AddIntAttribute(IImgMediaModule::CustomFormatAttributeName.Resolve().ToString(), 1);
-
-		// These attributes will not be added and therefore not found by EXR reader if it is not tiled.
-		if (bIsTiled)
-		{
-			OutFile.AddIntAttribute(IImgMediaModule::CustomFormatTileWidthAttributeName.Resolve().ToString(), TileWidth);
-			OutFile.AddIntAttribute(IImgMediaModule::CustomFormatTileHeightAttributeName.Resolve().ToString(), TileHeight);
-			OutFile.AddIntAttribute(IImgMediaModule::CustomFormatTileBorderAttributeName.Resolve().ToString(), InTileBorder);
-		}
-	}
+		DataWindow.Min, DataWindow.Max, bIsTiled);
 
 	// Add channels.
 	if (DestNumChannels == 4)
@@ -494,15 +478,7 @@ void SImgMediaProcessEXR::ProcessImageCustomRawData(TArray64<uint8>& RawData,
 		OutFile.AddChannel(RChannelName);
 	}
 
-	// Create output.
-	if (bIsCustomFormat)
-	{
-		OutFile.CreateOutputFile(InName, DestWidth, DestHeight, bInEnableMips, 1);
-	}
-	else
-	{
-		OutFile.CreateOutputFile(InName, TileWidth, TileHeight, bInEnableMips, 1);
-	}
+	OutFile.CreateOutputFile(InName, TileWidth, TileHeight, bInEnableMips, 1);
 
 	if (DestNumChannels == 4)
 	{
@@ -587,76 +563,9 @@ void SImgMediaProcessEXR::ProcessImageCustomRawData(TArray64<uint8>& RawData,
 			CurrentBuffer = TintBuffer.GetData();
 		}
 
-		// Do we need to tile this mip?
-		// Need to also check that this is actually a valid mip level.
-		if ((bIsCustomFormat) && (bIsTiled) && (MipSourceWidth > 0) && (MipSourceHeight > 0))
-		{
-			int32 MipTileWidth = TileWidth;
-			int32 MipTileHeight = TileHeight;
-
-			// A tile could be larger than the mip level when dealing with mips.
-			if (MipTileWidth > MipSourceWidth)
-			{
-				MipTileWidth = MipSourceWidth;
-			}
-			if (MipTileHeight > MipSourceHeight)
-			{
-				MipTileHeight = MipSourceHeight;
-			}
-
-			int32 OutputWidth = 0;
-			int32 OutputHeight = 0;
-			int32 MipNumTilesX = (MipSourceWidth + MipTileWidth - 1) / MipTileWidth;
-			int32 MipNumTilesY = (MipSourceHeight + MipTileHeight - 1) / MipTileHeight;
-
-			// Make sure our sizes match the mip size we get from EXR.
-			int32 ExpectedMipWidth = MipSourceWidth + MipNumTilesX * InTileBorder * 2;
-			if (ExpectedMipWidth != MipWidth)
-			{
-				UE_LOG(LogImgMediaEditor, Error,
-					TEXT("Expected mip level width of %d, but got %d (SourceWidth:%d NumTiles:%d TileBorder:%d"),
-					ExpectedMipWidth, MipHeight,
-					MipSourceWidth, MipNumTilesX, InTileBorder);
-			}
-			int32 ExpectedMipHeight = MipSourceHeight + MipNumTilesY * InTileBorder * 2;
-			if (ExpectedMipHeight != MipHeight)
-			{
-				UE_LOG(LogImgMediaEditor, Error,
-					TEXT("Expected mip level height of %d, but got %d (SourceHeight:%d NumTiles:%d TileBorder:%d"),
-					ExpectedMipHeight, MipHeight,
-					MipSourceHeight, MipNumTilesY, InTileBorder);
-			}
-
-			// Tile the buffer.
-			TileData(CurrentBuffer, TileBuffer,
-				MipSourceWidth, MipSourceHeight, MipWidth, MipHeight,
-				MipNumTilesX, MipNumTilesY,
-				MipTileWidth, MipTileHeight, InTileBorder,
-				BytesPerPixel);
-			CurrentBuffer = TileBuffer.GetData();
-		}
-
 		// Write to EXR.
 		TRACE_CPUPROFILER_EVENT_SCOPE(SImgMediaProcessEXR::ProcessImageCustom:WriteEXR);
-		if (bIsCustomFormat)
-		{
-			Stride.Y = MipWidth * BytesPerPixel;
-			int64 BufferOffset = 0;
-			int64 SingleBufferOffset = MipWidth * BytesPerPixelPerChannel;
-			if (DestNumChannels == 4)
-			{
-				OutFile.UpdateFrameBufferChannel(AChannelName, CurrentBuffer, Stride);
-				BufferOffset += SingleBufferOffset;
-			}
-
-			OutFile.UpdateFrameBufferChannel(BChannelName, CurrentBuffer + BufferOffset, Stride);
-			BufferOffset += SingleBufferOffset;
-			OutFile.UpdateFrameBufferChannel(GChannelName, CurrentBuffer + BufferOffset, Stride);
-			BufferOffset += SingleBufferOffset;
-			OutFile.UpdateFrameBufferChannel(RChannelName, CurrentBuffer + BufferOffset, Stride);
-			BufferOffset += SingleBufferOffset;
-		}
-		else
+		
 		{
 			Stride.X = BytesPerPixel;
 			Stride.Y = MipWidth * BytesPerPixel;
@@ -680,11 +589,6 @@ void SImgMediaProcessEXR::ProcessImageCustomRawData(TArray64<uint8>& RawData,
 
 		OutFile.SetFrameBuffer();
 
-		if (bIsCustomFormat)
-		{
-			OutFile.WriteTile(0, 0, MipLevel);
-		}
-		else
 		{
 			int32 X2 = OutFile.GetNumXTiles(MipLevel) - 1;
 			X2 = FMath::Max(0, X2);
@@ -956,7 +860,7 @@ void SImgMediaProcessEXR::HandleProcessing()
 						{
 							ProcessImageCustomRawData(RawData, Width, Height,
 								InTileWidth, InTileHeight, TileBorder, bEnableMips,
-								GlobalImageParameters, Name, false /*bUseCustomFormat*/);
+								GlobalImageParameters, Name);
 						});
 					}
 					else
diff --git a/Engine/Plugins/Media/ImgMedia/Source/ImgMediaEditor/Private/Widgets/SImgMediaProcessEXR.h b/Engine/Plugins/Media/ImgMedia/Source/ImgMediaEditor/Private/Widgets/SImgMediaProcessEXR.h
--- a/Engine/Plugins/Media/ImgMedia/Source/ImgMediaEditor/Private/Widgets/SImgMediaProcessEXR.h
+++ b/Engine/Plugins/Media/ImgMedia/Source/ImgMediaEditor/Private/Widgets/SImgMediaProcessEXR.h
@@ -98,12 +98,10 @@ private:
 	 * @param bInEnableMips		Turn on mip mapping.
 	 * @param ImageParameters	Contains information on the image.
 	 * @param InName			Full path and name of file to write.
-	 * @param bIsCustomFormat	True to output in our custom format.
 	 */
 	void ProcessImageCustom(const FImage& InImage,
 		int32 InTileWidth, int32 InTileHeight, int32 InTileBorder, 
-		bool bInEnableMips, const FImageParameters& ImageParameters, const FString& InName,
-		bool bIsCustomFormat);
+		bool bInEnableMips, const FImageParameters& ImageParameters, const FString& InName);
 
 	/**
 	 * Processess a single image and writes out a file.
@@ -119,12 +117,11 @@ private:
 	 * @param bInEnableMips		Turn on mip mapping.
 	 * @param ImageParameters	Contains information on the image.
 	 * @param InName			Full path and name of file to write.
-	 * @param bIsCustomFormat	True to output in our custom format.
 	 */
 	void ProcessImageCustomRawData(TArray64<uint8>& RawData,
 		int32 Width, int32 Height,
 		int32 InTileWidth, int32 InTileHeight, int32 InTileBorder, bool bInEnableMips,
-		const FImageParameters& ImageParameters, const FString& InName, bool bIsCustomFormat);
+		const FImageParameters& ImageParameters, const FString& InName);
 
 	/**
 	 * Removes the alpha channel from a buffer.
