EXR: Fix for pole 12k sequence pole artifacts.

diff --git a/Engine/Plugins/Media/ImgMedia/Source/ImgMedia/Private/Assets/ImgMediaMipMapInfo.cpp b/Engine/Plugins/Media/ImgMedia/Source/ImgMedia/Private/Assets/ImgMediaMipMapInfo.cpp
--- a/Engine/Plugins/Media/ImgMedia/Source/ImgMedia/Private/Assets/ImgMediaMipMapInfo.cpp
+++ b/Engine/Plugins/Media/ImgMedia/Source/ImgMedia/Private/Assets/ImgMediaMipMapInfo.cpp
@@ -754,6 +754,7 @@ namespace {
 								// Queue upscaling of lower quality mips into this mip.
 								MipLevelToUpscale = FMath::Max(PoleMipBias, MipLevelToUpscaleExcludingPoles);
 							}
+
 							float CalculatedLevel;
 							FIntVector2 MipLevelRange;
 
@@ -764,16 +765,26 @@ namespace {
 							{
 								CalculatedLevel += MipMapBias + ViewInfo.MaterialTextureMipBias;
 
-								MipLevelRange[0] = FMath::Clamp(FMath::FloorToInt32(CalculatedLevel - MipMapLevelPadding), 0, MaxLevel);
+								MipLevelRange[0] = FMath::FloorToInt32(CalculatedLevel - MipMapLevelPadding);
 								MipLevelRange[1] = FMath::CeilToInt32(CalculatedLevel + MipMapLevelPadding);
 
-								// As a mitigation for discontinuities at the poles, we artifically increase the max calculated level.
-								// (Note: Using an icosphere would avoid this issue but conflict with the partial sphere feature.)
-								if (TileY == 0 || TileY == SequenceTileNum.Y - 1)
+								const bool bIsTileAtPole = (TileY == 0 || TileY == SequenceTileNum.Y - 1);
+								/*
+								 * Since we estimate mip levels only once per tile here, this method very much approximates the
+								 * per-fragment hardware mip selection that occurs during rasterization. Discontinuities can
+								 * therefore appear where variance within one tile is strongest: the approximation delta is
+								 * indeed larger at the poles due to latlong spherical projection. So as a mitigation for pole
+								 * discontinuities, we artifically increase their mip level range, and accept the added read cost.
+								 *
+								 * NOTE: This should now be disabled by default since bAdaptivePoleMipUpscaling defaults to true.
+								 */
+								if (!bAdaptivePoleMipUpscaling && bIsTileAtPole)
 								{
+									MipLevelRange[0]--;
 									MipLevelRange[1]++;
 								}
 
+								MipLevelRange[0] = FMath::Clamp(MipLevelRange[0], 0, MaxLevel);
 								MipLevelRange[1] = FMath::Clamp(MipLevelRange[1], 0, MaxLevel);
 
 								for (int32 Level = MipLevelRange[0]; Level <= MipLevelRange[1]; ++Level)
@@ -786,6 +797,7 @@ namespace {
 											continue;
 										}
 									}
+
 									if (!VisibleTiles.Contains(Level))
 									{
 										VisibleTiles.Emplace(Level, FImgMediaTileSelection::CreateForTargetMipLevel(InSequenceInfo.Dim, InSequenceInfo.TilingDescription.TileSize, Level, false));
