ImgMedia: fixing up an issue when switching on and off looped playabck during playback after looping once or more times already & added clamp for incoming fetch time ranges
#rnx
#tests Local machine. MFW Test app, HUD level

diff --git a/Engine/Plugins/Media/ImgMedia/Source/ImgMedia/Private/Loader/ImgMediaLoader.cpp b/Engine/Plugins/Media/ImgMedia/Source/ImgMedia/Private/Loader/ImgMediaLoader.cpp
--- a/Engine/Plugins/Media/ImgMedia/Source/ImgMedia/Private/Loader/ImgMediaLoader.cpp
+++ b/Engine/Plugins/Media/ImgMedia/Source/ImgMedia/Private/Loader/ImgMediaLoader.cpp
@@ -390,6 +390,13 @@ IMediaSamples::EFetchBestSampleResult FImgMediaLoader::FetchBestVideoSampleForTi
 
 		auto TimeRange = InTimeRange;
 
+		// To avoid odd behavior with out of range time values, clamp them down just in case
+		if (!bIsLoopingEnabled)
+		{
+			TimeRange.SetLowerBoundValue(FMediaTimeStamp(FMath::Clamp(TimeRange.GetLowerBoundValue().Time, FTimespan::Zero(), SequenceDuration), TimeRange.GetLowerBoundValue().SequenceIndex));
+			TimeRange.SetUpperBoundValue(FMediaTimeStamp(FMath::Clamp(TimeRange.GetUpperBoundValue().Time, FTimespan::Zero(), SequenceDuration), TimeRange.GetUpperBoundValue().SequenceIndex));
+		}
+
 		// Clamp the current range requested against the last known timestamp to ensure we return the "next" sample at all times...
 		if (PlayRate >= 0.0f)
 		{
@@ -530,7 +537,7 @@ IMediaSamples::EFetchBestSampleResult FImgMediaLoader::FetchBestVideoSampleForTi
 		{
 			MaxLoopIdx = 0;	// ignored below
 			// Clamp this between [0..N-1], but keep a look out for any loop index that might persist from an earlier looping state
-			MaxIdx = FMath::Clamp(MaxIdx, StartLoopIndex * NumFrames, StartLoopIndex * NumFrames + NumFrames - 1);
+			MaxIdx = FMath::Clamp(MaxIdx, StartLoopIndex * NumFrames, StartLoopIndex * NumFrames + NumFrames - 1) - StartLoopIndex * NumFrames;
 		}
 
 		const TSharedPtr<FImgMediaFrame, ESPMode::ThreadSafe>* Frame;
