ImgMedia: Fixed additional blocking playback and deallocation issues with tiled EXR sequences.

* Reduced blocking waits on warmup sequence by only waiting for each buffer allocation individually
* Fixed missing deallocation of tracking mip map info objects due to sequencer calling SetMediaPlayer(nullptr)
* Fixed incorrect use of WeakPtr
* Fixed uninitialized MipLevelToUpscale, and incorrect logging

diff --git a/Engine/Plugins/Media/ImgMedia/Source/ImgMedia/Internal/Player/ImgMediaPlayer.h b/Engine/Plugins/Media/ImgMedia/Source/ImgMedia/Internal/Player/ImgMediaPlayer.h
--- a/Engine/Plugins/Media/ImgMedia/Source/ImgMedia/Internal/Player/ImgMediaPlayer.h
+++ b/Engine/Plugins/Media/ImgMedia/Source/ImgMedia/Internal/Player/ImgMediaPlayer.h
@@ -9,6 +9,8 @@
 #include "IMediaSamples.h"
 #include "IMediaTracks.h"
 #include "IMediaView.h"
+#include "UObject/WeakObjectPtr.h"
+#include "UObject/WeakObjectPtrTemplates.h"
 
 class FImgMediaLoader;
 class FImgMediaScheduler;
@@ -16,6 +18,7 @@ class IImgMediaReader;
 class IMediaEventSink;
 class IMediaTextureSample;
 class FImgMediaGlobalCache;
+class UMediaTexture;
 
 
 /**
@@ -183,4 +186,7 @@ private:
 
 	/** True if facade has signaled it uses blocking playback, false if not */
 	bool PlaybackIsBlocking;
+
+	/** Weak pointer to the active media texture. */
+	TWeakObjectPtr<UMediaTexture> MediaTextureWeakPtr;
 };
diff --git a/Engine/Plugins/Media/ImgMedia/Source/ImgMedia/Private/Assets/ImgMediaMipMapInfo.cpp b/Engine/Plugins/Media/ImgMedia/Source/ImgMedia/Private/Assets/ImgMediaMipMapInfo.cpp
--- a/Engine/Plugins/Media/ImgMedia/Source/ImgMedia/Private/Assets/ImgMediaMipMapInfo.cpp
+++ b/Engine/Plugins/Media/ImgMedia/Source/ImgMedia/Private/Assets/ImgMediaMipMapInfo.cpp
@@ -286,9 +286,10 @@ UMeshComponent* FImgMediaMipMapObjectInfo::GetMeshComponent() const
 
 int32 FImgMediaMipMapObjectInfo::GetMipLevelToUpscale() const
 {
-	if (Tracker.IsValid())
+	TSharedPtr<FMediaTextureTrackerObject> PinnedTracker = Tracker.Pin();
+	if (PinnedTracker.IsValid())
 	{
-		return Tracker.Pin()->MipLevelToUpscale;
+		return PinnedTracker->MipLevelToUpscale;
 	}
 
 	return -1;
@@ -624,6 +625,7 @@ namespace {
 		FSphereObjectInfo(UMeshComponent* InMeshComponent, TWeakPtr<FMediaTextureTrackerObject, ESPMode::ThreadSafe> InTracker)
 			: FImgMediaMipMapObjectInfo(InMeshComponent, MoveTemp(InTracker))
 			, DefaultSphereRadius(50.0f) // as defined in FMediaPlateCustomizationMesh::GenerateSphereMesh
+			, MipLevelToUpscale(-1)
 		{
 		}
 
@@ -670,9 +672,10 @@ namespace {
 			// Does user want to reduce the load at the poles?
 			bool bAdaptivePoleMipUpscaling = false;
 			MipLevelToUpscale = -1;
-			if (Tracker.IsValid())
+
+			TSharedPtr<FMediaTextureTrackerObject, ESPMode::ThreadSafe> PinnedTracker = Tracker.Pin();
+			if (PinnedTracker.IsValid())
 			{
-				TSharedPtr<FMediaTextureTrackerObject, ESPMode::ThreadSafe> PinnedTracker = Tracker.Pin();
 				MipLevelToUpscaleExcludingPoles = PinnedTracker->MipLevelToUpscale;
 				bAdaptivePoleMipUpscaling = PinnedTracker->bAdaptivePoleMipUpscaling;
 				MipLevelToUpscale = MipLevelToUpscaleExcludingPoles;
@@ -863,7 +866,9 @@ namespace {
 		}
 
 		const float DefaultSphereRadius;
+	
 	private:
+
 		mutable int32 MipLevelToUpscale;
 	};
 
@@ -903,6 +908,8 @@ void FImgMediaMipMapInfo::AddObject(AActor* InActor, TWeakPtr<FMediaTextureTrack
 				Objects.Add(new FImgMediaMipMapObjectInfo(MeshComponent, InTracker));
 				break;
 			}
+
+			UE_LOG(LogImgMedia, Verbose, TEXT("Added tracked %s for %s"), *MeshComponent->GetName(), *InActor->GetName());
 		}
 
 		SubscribeEndFrame();
@@ -923,6 +930,8 @@ void FImgMediaMipMapInfo::RemoveObject(AActor* InActor)
 			{
 				if (InActor == MeshComponent->GetOuter())
 				{
+					UE_LOG(LogImgMedia, Verbose, TEXT("Removed tracked %s for %s"), *MeshComponent->GetName(), *InActor->GetName());
+
 					Objects.RemoveAtSwap(Index);
 					delete Info;
 
diff --git a/Engine/Plugins/Media/ImgMedia/Source/ImgMedia/Private/Loader/ImgMediaLoader.cpp b/Engine/Plugins/Media/ImgMedia/Source/ImgMedia/Private/Loader/ImgMediaLoader.cpp
--- a/Engine/Plugins/Media/ImgMedia/Source/ImgMedia/Private/Loader/ImgMediaLoader.cpp
+++ b/Engine/Plugins/Media/ImgMedia/Source/ImgMedia/Private/Loader/ImgMediaLoader.cpp
@@ -699,6 +699,14 @@ IQueuedWork* FImgMediaLoader::GetWork()
 	// Update minimum mip level to upscale
 	MinimumMipLevelToUpscale = GetDesiredMinimumMipLevelToUpscale();
 
+	if (MinimumMipLevelToUpscale >= 0 && NumMipLevels <= 1)
+	{
+		UE_CALL_ONCE([ImagePath = GetImagePath(0, 0)]
+			{
+				UE_LOG(LogImgMedia, Display, TEXT("No upscaling for sequence without mips: %s"), *ImagePath);
+			});
+	}
+
 	FImgMediaLoaderWork* Work = (WorkPool.Num() > 0) ? WorkPool.Pop() : new FImgMediaLoaderWork(AsShared(), Reader.ToSharedRef());
 	
 	// Get the existing frame so we can add the mip level to it.
diff --git a/Engine/Plugins/Media/ImgMedia/Source/ImgMedia/Private/Player/ImgMediaPlayer.cpp b/Engine/Plugins/Media/ImgMedia/Source/ImgMedia/Private/Player/ImgMediaPlayer.cpp
--- a/Engine/Plugins/Media/ImgMedia/Source/ImgMedia/Private/Player/ImgMediaPlayer.cpp
+++ b/Engine/Plugins/Media/ImgMedia/Source/ImgMedia/Private/Player/ImgMediaPlayer.cpp
@@ -35,38 +35,6 @@ DECLARE_CYCLE_STAT(TEXT("ImgMedia Player TickInput"), STAT_ImgMedia_PlayerTickIn
 
 const FTimespan HackDeltaTimeOffset(1);
 
-namespace {
-	/** Convenience function to process all media textures corresponding the specified player. */
-	void ApplyToPlayerMediaTextures(FImgMediaPlayer* InPlayer, TFunctionRef<void(UMediaTexture*)> TextureCallbackFn)
-	{
-		TArray<UMediaTexture*> PlayerTextures;
-
-		FMediaTextureTracker& TextureTracker = FMediaTextureTracker::Get();
-
-		// Look through all the media textures we know about.
-		for (TWeakObjectPtr<UMediaTexture> TexturePtr : TextureTracker.GetTextures())
-		{
-			UMediaTexture* Texture = TexturePtr.Get();
-			if (Texture != nullptr)
-			{
-				// Does this match the player?
-				UMediaPlayer* MediaPlayer = Texture->GetMediaPlayer();
-				if (MediaPlayer != nullptr)
-				{
-					TSharedPtr<IMediaPlayer, ESPMode::ThreadSafe> Player = MediaPlayer->GetPlayerFacade()->GetPlayer();
-					if (Player.IsValid())
-					{
-						if (Player.Get() == InPlayer)
-						{
-							TextureCallbackFn(Texture);
-						}
-					}
-				}
-			}
-		}
-	}
-}
-
 
 /* FImgMediaPlayer structors
  *****************************************************************************/
@@ -110,15 +78,10 @@ void FImgMediaPlayer::Close()
 		return;
 	}
 
+	const TSharedPtr<FImgMediaMipMapInfo, ESPMode::ThreadSafe>& MipMapInfo = Loader->GetMipMapInfo();
+	if (MipMapInfo.IsValid() && MediaTextureWeakPtr.IsValid())
 	{
-		const TSharedPtr<FImgMediaMipMapInfo, ESPMode::ThreadSafe>& MipMapInfo = Loader->GetMipMapInfo();
-		if (MipMapInfo.IsValid())
-		{
-			ApplyToPlayerMediaTextures(this, [&MipMapInfo](UMediaTexture* Texture)
-				{
-					MipMapInfo->RemoveObjectsUsingThisMediaTexture(Texture);
-				});
-		}
+		MipMapInfo->RemoveObjectsUsingThisMediaTexture(MediaTextureWeakPtr.Get());
 	}
 
 	Scheduler->UnregisterLoader(Loader.ToSharedRef());
@@ -275,10 +238,29 @@ bool FImgMediaPlayer::Open(const FString& Url, const IMediaOptions* Options)
 			MipMapInfo = StaticCastSharedPtr<FImgMediaMipMapInfo, IMediaOptions::FDataContainer, ESPMode::ThreadSafe>(DataContainer);
 			if (MipMapInfo.IsValid())
 			{
-				ApplyToPlayerMediaTextures(this, [&MipMapInfo](UMediaTexture* Texture)
+				// Look through all the media textures we know about to find the one tied to this player
+				for (TWeakObjectPtr<UMediaTexture> TexturePtr : FMediaTextureTracker::Get().GetTextures())
+				{
+					UMediaTexture* Texture = TexturePtr.Get();
+					if (Texture != nullptr)
 					{
-						MipMapInfo->AddObjectsUsingThisMediaTexture(Texture);
-					});
+						// Does this match the player?
+						UMediaPlayer* MediaPlayer = Texture->GetMediaPlayer();
+						if (MediaPlayer != nullptr)
+						{
+							TSharedPtr<IMediaPlayer, ESPMode::ThreadSafe> Player = MediaPlayer->GetPlayerFacade()->GetPlayer();
+							if (Player.IsValid())
+							{
+								if (Player.Get() == this)
+								{
+									MediaTextureWeakPtr = TexturePtr;
+
+									MipMapInfo->AddObjectsUsingThisMediaTexture(MediaTextureWeakPtr.Get());
+								}
+							}
+						}
+					}
+				}
 			}
 		}
 	}
diff --git a/Engine/Plugins/Media/ImgMedia/Source/ImgMedia/Private/Readers/ExrImgMediaReaderGpu.cpp b/Engine/Plugins/Media/ImgMedia/Source/ImgMedia/Private/Readers/ExrImgMediaReaderGpu.cpp
--- a/Engine/Plugins/Media/ImgMedia/Source/ImgMedia/Private/Readers/ExrImgMediaReaderGpu.cpp
+++ b/Engine/Plugins/Media/ImgMedia/Source/ImgMedia/Private/Readers/ExrImgMediaReaderGpu.cpp
@@ -86,12 +86,7 @@ FExrImgMediaReaderGpu::FExrImgMediaReaderGpu(const TSharedRef<FImgMediaLoader, E
 	, bIsShuttingDown(false)
 	, bFallBackToCPU(false)
 {
-	const int32 NumMipLevels = InLoader->GetNumMipLevels();
 
-	if (NumMipLevels <= 1 && InLoader->GetMinimumLevelToUpscale() >= 0)
-	{
-		UE_LOG(LogImgMedia, Display, TEXT("No upscaling for sequence without mips: %s"), *InLoader->GetImagePath(0,0));
-	}
 }
 
 FExrImgMediaReaderGpu::~FExrImgMediaReaderGpu()
@@ -380,7 +375,7 @@ void FExrImgMediaReaderGpu::PreAllocateMemoryPool(int32 NumFrames, const FImgMed
 	SIZE_T AllocSize = GetBufferSize(FrameInfo.Dim, FrameInfo.NumChannels, FrameInfo.bHasTiles, FrameInfo.NumTiles, bCustomExr);
 	for (int32 FrameCacheNum = 0; FrameCacheNum < NumFrames; FrameCacheNum++)
 	{
-		AllocateGpuBufferFromPool(AllocSize, FrameCacheNum == NumFrames - 1);
+		AllocateGpuBufferFromPool(AllocSize);
 	}
 }
 
@@ -613,7 +608,7 @@ void FExrImgMediaReaderGpu::CreateSampleConverterCallback(TSharedPtr<FExrMediaTe
 	SampleConverter->AddCallback(FExrConvertBufferCallback::CreateLambda(RenderThreadSwizzler));
 }
 
-FStructuredBufferPoolItemSharedPtr FExrImgMediaReaderGpu::AllocateGpuBufferFromPool(uint32 AllocSize, bool bWait)
+FStructuredBufferPoolItemSharedPtr FExrImgMediaReaderGpu::AllocateGpuBufferFromPool(uint32 AllocSize)
 {
 	TRACE_CPUPROFILER_EVENT_SCOPE_TEXT(*FString::Printf(TEXT("ExrReaderGpu.AllocBuffer")));
 	TRACE_CPUPROFILER_EVENT_SCOPE_TEXT(*FString::Printf(TEXT("ExrReaderGpu.AllocBuffer %d"), AllocSize));
@@ -650,21 +645,14 @@ FStructuredBufferPoolItemSharedPtr FExrImgMediaReaderGpu::AllocateGpuBufferFromP
 
 	if (!AllocatedBuffer)
 	{
-		FEvent* InitDoneEvent = nullptr;
-		if (bWait)
-		{
-			InitDoneEvent = FPlatformProcess::GetSynchEventFromPool();
-		}
+		AllocatedBuffer = MakeShareable(new FStructuredBufferPoolItem(), MoveTemp(BufferDeleter));
 
-		{
-			AllocatedBuffer = MakeShareable(new FStructuredBufferPoolItem(), MoveTemp(BufferDeleter));
-
-			// Allocate and unlock the structured buffer on render thread.
-			ENQUEUE_RENDER_COMMAND(CreatePooledBuffer)([AllocatedBuffer, AllocSize, InitDoneEvent](FRHICommandListImmediate& RHICmdList)
+		// Allocate and unlock the structured buffer on render thread.
+		ENQUEUE_RENDER_COMMAND(CreatePooledBuffer)([AllocatedBuffer, AllocSize](FRHICommandListImmediate& RHICmdList)
 			{
 				TRACE_CPUPROFILER_EVENT_SCOPE_TEXT(*FString::Printf(TEXT("ExrReaderGpu.AllocBuffer_RenderThread")));
 				TRACE_CPUPROFILER_EVENT_SCOPE_TEXT(*FString::Printf(TEXT("ExrReaderGpu.AllocBuffer_RenderThread %d"), AllocSize));
-				
+
 				SCOPED_GPU_STAT(RHICmdList, ExrImgMediaReaderGpu_AllocateBuffer);
 				SCOPED_DRAW_EVENT(RHICmdList, FExrImgMediaReaderGpu_AllocateBuffer);
 
@@ -684,19 +672,8 @@ FStructuredBufferPoolItemSharedPtr FExrImgMediaReaderGpu::AllocateGpuBufferFromP
 					AllocatedBuffer->ShaderResourceView = RHICmdList.CreateShaderResourceView(AllocatedBuffer->UploadBufferRef);
 				}
 
-				if (InitDoneEvent != nullptr)
-				{
-					InitDoneEvent->Trigger();
-				}
+				AllocatedBuffer->AllocationReadyEvent->Trigger();
 			});
-		}
-
-		/** Wait until buffer is initialized. */
-		if (InitDoneEvent != nullptr)
-		{
-			InitDoneEvent->Wait();
-			FPlatformProcess::ReturnSynchEventToPool(InitDoneEvent);
-		}
 	}
 
 	// This buffer will be automatically processed and returned to StagingMemoryPool once nothing keeps reference to it.
@@ -725,11 +702,20 @@ bool FExrMediaTextureSampleConverter::Convert(FTexture2DRHIRef& InDstTexture, co
 	return bExecutionSuccessful;
 }
 
+FStructuredBufferPoolItem::FStructuredBufferPoolItem()
+{
+	constexpr bool bIsManualReset = true; // Manually reset events stay triggered until reset.
+	AllocationReadyEvent = FPlatformProcess::GetSynchEventFromPool(bIsManualReset);
+	check(AllocationReadyEvent);
+}
+
 FStructuredBufferPoolItem::~FStructuredBufferPoolItem()
 {
 	TRACE_CPUPROFILER_EVENT_SCOPE_TEXT(*FString::Printf(TEXT("ExrReaderGpu.ReleasePoolItem")));
 	FRHICommandListImmediate::Get().UnlockBuffer(UploadBufferRef);
 	UploadBufferMapped = nullptr;
+
+	FPlatformProcess::ReturnSynchEventToPool(AllocationReadyEvent);
 }
 
 #endif //IMGMEDIA_EXR_SUPPORTED_PLATFORM && PLATFORM_WINDOWS
diff --git a/Engine/Plugins/Media/ImgMedia/Source/ImgMedia/Private/Readers/ExrImgMediaReaderGpu.h b/Engine/Plugins/Media/ImgMedia/Source/ImgMedia/Private/Readers/ExrImgMediaReaderGpu.h
--- a/Engine/Plugins/Media/ImgMedia/Source/ImgMedia/Private/Readers/ExrImgMediaReaderGpu.h
+++ b/Engine/Plugins/Media/ImgMedia/Source/ImgMedia/Private/Readers/ExrImgMediaReaderGpu.h
@@ -36,9 +36,13 @@ struct FStructuredBufferPoolItem
 	*/
 	FShaderResourceViewRHIRef ShaderResourceView;
 
-	/**
-	* Destructure to clean up render resources
-	*/
+	/** Event used to wait for completed buffer allocations. */
+	FEvent* AllocationReadyEvent;
+
+	/* Constructor */
+	FStructuredBufferPoolItem();
+
+	/** Destructor to clean up render resources */
 	~FStructuredBufferPoolItem();
 };
 
@@ -105,7 +109,7 @@ protected:
 public:
 
 	/** Typically we would need the (ImageResolution.x*y)*NumChannels*ChannelSize */
-	virtual FStructuredBufferPoolItemSharedPtr AllocateGpuBufferFromPool(uint32 AllocSize, bool bWait = true);
+	virtual FStructuredBufferPoolItemSharedPtr AllocateGpuBufferFromPool(uint32 AllocSize);
 
 	/** Either return or Add new chunk of memory to the pool based on its size. */
 	void ReturnGpuBufferToPool(uint32 AllocSize, FStructuredBufferPoolItem* Buffer);
@@ -202,13 +206,25 @@ public:
 
 	FStructuredBufferPoolItemSharedPtr GetOrCreateMipLevelBuffer(int32 RequestedMipLevel, TFunction<FStructuredBufferPoolItemSharedPtr()> AllocatorFunc)
 	{
-		FScopeLock Lock(&MipBufferCriticalSection);
-		if (FStructuredBufferPoolItemSharedPtr* BufferDataPtr = MipBuffers.Find(RequestedMipLevel))
+		FStructuredBufferPoolItemSharedPtr Result;
+
 		{
-			return *BufferDataPtr;
+			FScopeLock Lock(&MipBufferCriticalSection);
+
+			if (FStructuredBufferPoolItemSharedPtr* BufferDataPtr = MipBuffers.Find(RequestedMipLevel))
+			{
+				Result = *BufferDataPtr;
+			}
+			else
+			{
+				Result = MipBuffers.Add(RequestedMipLevel, AllocatorFunc());
+			}
 		}
+		
+		// Wait for render thread buffer allocations before using resources
+		Result->AllocationReadyEvent->Wait();
 
-		return MipBuffers.Add(RequestedMipLevel, AllocatorFunc());
+		return Result;
 	}
 	
 	FSampleConverterParameters GetParams()
