ImgMedia: Fixed crash with EXR mips in separate folders.

Using the mip getter on the loader handles both in-frame mips and mip-per-folder cases. Regression was introduced in CL 29237305.

#rnx

diff --git a/Engine/Plugins/Media/ImgMedia/Source/ImgMedia/Private/Readers/ExrImgMediaReader.cpp b/Engine/Plugins/Media/ImgMedia/Source/ImgMedia/Private/Readers/ExrImgMediaReader.cpp
--- a/Engine/Plugins/Media/ImgMedia/Source/ImgMedia/Private/Readers/ExrImgMediaReader.cpp
+++ b/Engine/Plugins/Media/ImgMedia/Source/ImgMedia/Private/Readers/ExrImgMediaReader.cpp
@@ -207,7 +207,7 @@ bool FExrImgMediaReader::ReadFrame(int32 FrameId, const TMap<int32, FImgMediaTil
 		// allocate frame buffer
 		// FRgbaInputFile loads any exr as rgba 16 bit per channel.
 		// UncompressedSize is used for cache, therefore it needs to be set to the actual size of the buffer.
-		FrameInfo.UncompressedSize = GetMipBufferTotalSize(Dim, FrameInfo.NumMipLevels > 1);
+		FrameInfo.UncompressedSize = GetMipBufferTotalSize(Dim, Loader->GetNumMipLevels() > 1);
 		SIZE_T BufferSize = FrameInfo.UncompressedSize;
 		void* Buffer = FMemory::Malloc(BufferSize, PLATFORM_CACHE_LINE_SIZE);
 
