ImgMedia: Fixed issue with player never releasing mip map info objects.

* Stop the loader from resetting its mip map info shared ptr. This previously caused the image media player to never call RemoveObjectsUsingThisMediaTexture.
* Fixed potential double scaling of UncompressedSize: once from players, once from the loader.

#rnx

diff --git a/Engine/Plugins/Media/ImgMedia/Source/ImgMedia/Internal/Loader/ImgMediaLoader.h b/Engine/Plugins/Media/ImgMedia/Source/ImgMedia/Internal/Loader/ImgMediaLoader.h
--- a/Engine/Plugins/Media/ImgMedia/Source/ImgMedia/Internal/Loader/ImgMediaLoader.h
+++ b/Engine/Plugins/Media/ImgMedia/Source/ImgMedia/Internal/Loader/ImgMediaLoader.h
@@ -570,6 +570,11 @@ protected:
 	 */
 	void ResetFetchLogic();
 
+	/**
+	 * Check whether the current sequence has valid mip levels or tiles, and a valid mip map info object.
+	 */
+	bool HasMipsOrTilesEnabled() const;
+
 private:
 
 	/** Critical section for synchronizing access to Frames. */
diff --git a/Engine/Plugins/Media/ImgMedia/Source/ImgMedia/Private/GlobalCache/ImgMediaGlobalCache.cpp b/Engine/Plugins/Media/ImgMedia/Source/ImgMedia/Private/GlobalCache/ImgMediaGlobalCache.cpp
--- a/Engine/Plugins/Media/ImgMedia/Source/ImgMedia/Private/GlobalCache/ImgMediaGlobalCache.cpp
+++ b/Engine/Plugins/Media/ImgMedia/Source/ImgMedia/Private/GlobalCache/ImgMediaGlobalCache.cpp
@@ -60,13 +60,13 @@ void FImgMediaGlobalCache::Shutdown()
 #endif // WITH_EDITOR
 }
 
-void FImgMediaGlobalCache::AddFrame(const FString& FileName, const FName& Sequence, int32 Index, const TSharedPtr<FImgMediaFrame, ESPMode::ThreadSafe>& Frame, bool HasMipMaps)
+void FImgMediaGlobalCache::AddFrame(const FString& FileName, const FName& Sequence, int32 Index, const TSharedPtr<FImgMediaFrame, ESPMode::ThreadSafe>& Frame, bool bAdjustSizeForMipMaps)
 {
 	FScopeLock Lock(&CriticalSection);
 
 	// Make sure we have enough space in the cache to add this new frame.
 	SIZE_T FrameSize = Frame->GetUncompressedSize();
-	if (HasMipMaps)
+	if (bAdjustSizeForMipMaps)
 	{
 		FrameSize = (FrameSize * 4) / 3;
 	}
diff --git a/Engine/Plugins/Media/ImgMedia/Source/ImgMedia/Private/Loader/ImgMediaLoader.cpp b/Engine/Plugins/Media/ImgMedia/Source/ImgMedia/Private/Loader/ImgMediaLoader.cpp
--- a/Engine/Plugins/Media/ImgMedia/Source/ImgMedia/Private/Loader/ImgMediaLoader.cpp
+++ b/Engine/Plugins/Media/ImgMedia/Source/ImgMedia/Private/Loader/ImgMediaLoader.cpp
@@ -317,6 +317,11 @@ void FImgMediaLoader::ResetFetchLogic()
 	SkipFramesLevel = 0;
 }
 
+bool FImgMediaLoader::HasMipsOrTilesEnabled() const
+{
+	return MipMapInfo.IsValid() && (NumMipLevels > 1 || IsTiled());
+}
+
 void FImgMediaLoader::SetIsPlaybackBlocking(bool bIsBlocking)
 {
 	bIsPlaybackBlocking = bIsBlocking;
@@ -900,6 +905,13 @@ bool FImgMediaLoader::LoadSequence(const FString& SequencePath, const FFrameRate
 	SequenceDuration = FrameNumberToTime(GetNumImages());
 	SIZE_T UncompressedSize = FirstFrameInfo.UncompressedSize;
 
+	// With mips in separate folders, we scale the uncompressed size by its mip chain approximation.
+	// With mips in each frame, the readers already return a value that incorporates their size.
+	if (ImagePaths.Num() > 1)
+	{
+		UncompressedSize = (UncompressedSize * 4) / 3;
+	}
+
 	if (FirstFrameInfo.bHasTiles)
 	{
 		TilingDescription.TileBorderSize = FirstFrameInfo.TileBorder;
@@ -916,22 +928,10 @@ bool FImgMediaLoader::LoadSequence(const FString& SequencePath, const FFrameRate
 	}
 #endif
 
-	// If we have no mips or tiles, then get rid of our MipMapInfoObject.
-	// Otherwise, set it up.
-	if (MipMapInfo.IsValid())
+	// If we have mips or tiles, set up our MipMapInfoObject.
+	if (HasMipsOrTilesEnabled())
 	{
-		if ((GetNumMipLevels() == 1) && (IsTiled() == false))
-		{
-			MipMapInfo.Reset();
-		}
-		else
-		{
-			MipMapInfo->SetTextureInfo(SequenceName, GetNumMipLevels(), SequenceDim, TilingDescription);
-			if (GetNumMipLevels() > 1)
-			{
-				UncompressedSize = (UncompressedSize * 4) / 3;
-			}
-		}
+		MipMapInfo->SetTextureInfo(SequenceName, NumMipLevels, SequenceDim, TilingDescription);
 	}
 
 	// initialize loader
@@ -1467,7 +1467,8 @@ void FImgMediaLoader::AddFrameToCache(int32 FrameNumber, const TSharedPtr<FImgMe
 			ExistingFrame = GlobalCache->FindAndTouch(SequenceName, FrameNumber);
 			if (ExistingFrame == nullptr)
 			{
-				GlobalCache->AddFrame(ImagePaths[0][FrameNumber], SequenceName, FrameNumber, Frame, MipMapInfo.IsValid());
+				const bool bHasMipMapsInSeparateFolders = ImagePaths.Num() > 1;
+				GlobalCache->AddFrame(ImagePaths[0][FrameNumber], SequenceName, FrameNumber, Frame, bHasMipMapsInSeparateFolders);
 			}
 		}
 		else
@@ -1480,8 +1481,8 @@ void FImgMediaLoader::AddFrameToCache(int32 FrameNumber, const TSharedPtr<FImgMe
 void FImgMediaLoader::GetDesiredMipTiles(int32 FrameIndex, TMap<int32, FImgMediaTileSelection>& OutMipsAndTiles)
 {
 	// note: While currently unused, FrameIndex could become useful for recorded frames.
-
-	if(MipMapInfo.IsValid() && MipMapInfo->HasObjects())
+	
+	if(HasMipsOrTilesEnabled() && MipMapInfo->HasObjects())
 	{
 		OutMipsAndTiles = MipMapInfo->GetVisibleTiles();
 	}
@@ -1497,7 +1498,7 @@ void FImgMediaLoader::GetDesiredMipTiles(int32 FrameIndex, TMap<int32, FImgMedia
 
 int32 FImgMediaLoader::GetDesiredMinimumMipLevelToUpscale()
 {
-	if (MipMapInfo.IsValid())
+	if (HasMipsOrTilesEnabled())
 	{
 		return MipMapInfo->GetMinimumMipLevelToUpscale();
 	}
diff --git a/Engine/Plugins/Media/ImgMedia/Source/ImgMedia/Public/ImgMediaGlobalCache.h b/Engine/Plugins/Media/ImgMedia/Source/ImgMedia/Public/ImgMediaGlobalCache.h
--- a/Engine/Plugins/Media/ImgMedia/Source/ImgMedia/Public/ImgMediaGlobalCache.h
+++ b/Engine/Plugins/Media/ImgMedia/Source/ImgMedia/Public/ImgMediaGlobalCache.h
@@ -46,9 +46,9 @@ public:
 	 * @param	Sequence		Indentifying name of this sequence.
 	 * @param	Index			Index of frame to add.
 	 * @param	Frame			Actual frame to add.
-	 * @param	HasMipMaps		Whether this frame contains mipmaps or not.
+	 * @param	bAdjustSizeForMipMaps	Whether this frame needs size adjustments for mip maps in separate folders.
 	 */
-	void AddFrame(const FString& FileName, const FName& Sequence, int32 Index, const TSharedPtr<FImgMediaFrame, ESPMode::ThreadSafe>& Frame, bool HasMipMaps);
+	void AddFrame(const FString& FileName, const FName& Sequence, int32 Index, const TSharedPtr<FImgMediaFrame, ESPMode::ThreadSafe>& Frame, bool bAdjustSizeForMipMaps = false);
 
 	/**
 	 * Check whether the entry with the specified sequence and index is present.
