ImgMedia: Removing unused custom exr code paths.

diff --git a/Engine/Plugins/Media/ImgMedia/Shaders/Private/ExrSwizzler.usf b/Engine/Plugins/Media/ImgMedia/Shaders/Private/ExrSwizzler.usf
--- a/Engine/Plugins/Media/ImgMedia/Shaders/Private/ExrSwizzler.usf
+++ b/Engine/Plugins/Media/ImgMedia/Shaders/Private/ExrSwizzler.usf
@@ -2,6 +2,8 @@
 
 #include "/Engine/Private/Common.ush"
 
+#define NUM_CHANNELS (PERMUTATION_CHANNELS + 1)
+
 #if PARTIAL_TILES
 struct FTileDesc
 {
@@ -31,40 +33,21 @@ int2 NumTiles;
 /** Number of channels per pixel. */
 int NumChannels;
 
-/**
-* This function extracts the 16bit (float) data stored in uint (32 bit)
-*/
-float UintContainingFloatToFloat(in uint Value)
-{
-	int Sign = (Value >> 15) & 0x00000001;
-	int Exponent = (Value >> 10) & 0x0000001f;
-	int Mantissa = Value & 0x000003ff;
-
-	Exponent = Exponent + (127 - 15);
-	Mantissa = Mantissa << 13;
-
-	return asfloat((Sign << 31) | (Exponent << 23) | Mantissa);
-}
-
 /** Unpacks data from the provided structured buffer. */
-int GetValueFromTheBuffer(int Offset, int2 TextureSize, int ChannelIndex)
+half GetBufferValue(int Offset, int2 TextureSize, int ChannelIndex)
 {
-#if CUSTOM_EXR
-	int UnpackedLocation = Offset + ChannelIndex;
-#else
 	// the location as if it would've been if the buffer was unpacked.
 	int UnpackedLocation = Offset + TextureSize.x * ChannelIndex;
-#endif
 	// Byte offset within the 4 byte packed value.
 	int Remainder = UnpackedLocation % 2;
 
 	// Position in a packed buffer.
 	int PackedPosition = (UnpackedLocation - Remainder) >> 1;
 
+	uint PackedValue = (UnswizzledBuffer[PackedPosition] >> (16 * Remainder)) & 0xffff;
 
-	uint PackedFourByteValue = UnswizzledBuffer[PackedPosition];
-	int ReturnValue = ((PackedFourByteValue >> (16 * Remainder)) & 0xffff);
-	return ReturnValue;
+	// Convert from uint32 to f16
+	return (half)f16tof32(PackedValue);
 }
 
 
@@ -96,11 +79,7 @@ void GetStartScanlinePosition_PartialTiles(in int2 PixelPos, out int2 SampleSize
 	int PreviousTilesOffset = TileDesc.BufferOffset / 2;
 
 	// Current tile position.
-#if CUSTOM_EXR
-	StartPosition = PreviousTilesOffset + yCoord * TileDesc.TileResolution.x * NumChannels + xCoord * NumChannels;
-#else
 	StartPosition = PreviousTilesOffset + yCoord * TileDesc.TileResolution.x * NumChannels + xCoord + TilePadding;
-#endif //CUSTOM_EXR
 }
 #endif
 
@@ -120,15 +99,9 @@ void GetStartScanlinePosition(in int2 PixelPos, out int2 SampleSize, out int Sta
 	const int TilePadding = 10;
 
 	// Current tile position.
-#if CUSTOM_EXR
-	int TileBufferStride = TileSize.x * TileSize.y * NumChannels;
-	int PreviousTilesOffset = (NumTiles.x * TileY + TileX) * TileBufferStride;
-	StartPosition = PreviousTilesOffset + yCoord * TileSize.x * NumChannels + xCoord * NumChannels;
-#else
 	int TileBufferStride = TileSize.x * TileSize.y * NumChannels;
 	int PreviousTilesOffset = (NumTiles.x * TileY + TileX) * TileBufferStride;
 	StartPosition = PreviousTilesOffset + yCoord * TileSize.x * NumChannels + xCoord + (TileY * NumTiles.x + TileX + 1) * TilePadding;
-#endif
 
 #else
 	SampleSize.x = TextureSize.x;
@@ -161,33 +134,25 @@ half4 MainPS(noperspective float4 UVAndScreenPos : TEXCOORD0) : SV_Target0
 #endif
 
 	int ChannelIndex = 0;
+	half4 Color = half4(0.0f, 0.0f, 0.0f, 1.0f);
 
-#if NUM_CHANNELS == 3
-	half A = UintContainingFloatToFloat(GetValueFromTheBuffer(StartBufferPosition, SampleSize, ChannelIndex++));
-#else
-	half A = 1.;
+#if NUM_CHANNELS == 4
+	Color.a = GetBufferValue(StartBufferPosition, SampleSize, ChannelIndex++);
 #endif
-
-	half B = UintContainingFloatToFloat(GetValueFromTheBuffer(StartBufferPosition, SampleSize, ChannelIndex++));
-#if NUM_CHANNELS == 0
-	half G = B;
-#else
-	half G = UintContainingFloatToFloat(GetValueFromTheBuffer(StartBufferPosition, SampleSize, ChannelIndex++));
+#if NUM_CHANNELS >= 3
+	Color.b = GetBufferValue(StartBufferPosition, SampleSize, ChannelIndex++);
 #endif
-#if NUM_CHANNELS <= 1
-	half R = B;
-#else
-	half R = UintContainingFloatToFloat(GetValueFromTheBuffer(StartBufferPosition, SampleSize, ChannelIndex++));
+#if NUM_CHANNELS >= 2
+	Color.g = GetBufferValue(StartBufferPosition, SampleSize, ChannelIndex++);
 #endif
-	half3 Rgb = half3(R, G, B);
-
-#if CUSTOM_EXR
-	#if NUM_CHANNELS == 3
-		return half4(A, Rgb.zyx);
-	#else
-		return half4(Rgb.zyx, A);
-	#endif
-#else
-	return half4(Rgb, A);
+#if NUM_CHANNELS >= 1
+	Color.r = GetBufferValue(StartBufferPosition, SampleSize, ChannelIndex++);
 #endif
+
+	// Special case for single channel EXRs: we convert to grayscale by duplicating the red channel.
+#if NUM_CHANNELS == 1
+	Color.rgb = Color.r;
+#endif
+
+	return Color;
 }
\ No newline at end of file
diff --git a/Engine/Plugins/Media/ImgMedia/Source/ExrReaderGpu/Private/ExrReaderGpu.cpp b/Engine/Plugins/Media/ImgMedia/Source/ExrReaderGpu/Private/ExrReaderGpu.cpp
--- a/Engine/Plugins/Media/ImgMedia/Source/ExrReaderGpu/Private/ExrReaderGpu.cpp
+++ b/Engine/Plugins/Media/ImgMedia/Source/ExrReaderGpu/Private/ExrReaderGpu.cpp
@@ -261,126 +261,88 @@ bool FExrReader::GenerateTextureData(uint16* Buffer, int32 BufferSize, FString F
 
 void FExrReader::CalculateTileOffsets
 	( TArray<int32>& OutNumTilesPerLevel
-	, TArray<TArray<int64>>& OutCustomOffsets
 	, TArray<TArray<FTileDesc>>& OutPartialTileInfo
 	, const FIntPoint& FullTextureResolution
 	, const FIntPoint& TileDimWithBorders
 	, int32 NumMipLevels
-	, int64 PixelSize
-	, bool bCustomExr)
+	, int64 PixelSize)
 {
 	TRACE_CPUPROFILER_EVENT_SCOPE_TEXT(*FString::Printf(TEXT("CalculateTileOffsets")));
 
-	OutCustomOffsets.SetNum(NumMipLevels);
-	// Custom exr files don't have tile offsets, therefore we need to calculate these manually.
+	int64 CurrentPosition = 0;
+	for (int32 MipLevel = 0; MipLevel < NumMipLevels; MipLevel++)
 	{
-		int64 CurrentPosition = 0;
-		for (int32 MipLevel = 0; MipLevel < NumMipLevels; MipLevel++)
-		{
-			int32 MipDiv = 1 << MipLevel;
+		int32 MipDiv = 1 << MipLevel;
 
-			// Resolution of the texture in pixels for this mip level.
-			FIntPoint MipResolution = FullTextureResolution / MipDiv;
+		// Resolution of the texture in pixels for this mip level.
+		FIntPoint MipResolution = FullTextureResolution / MipDiv;
 
-			// Dimension of the texture in tiles including partial tiles
-			float TileResFractionX = float(MipResolution.X) / TileDimWithBorders.X;
-			float TileResFractionY = float(MipResolution.Y) / TileDimWithBorders.Y;
-			FIntPoint DimensionInTiles_PartialTiles
-				( FMath::CeilToInt(TileResFractionX)
-				, FMath::CeilToInt(TileResFractionY));
+		// Dimension of the texture in tiles including partial tiles
+		float TileResFractionX = float(MipResolution.X) / TileDimWithBorders.X;
+		float TileResFractionY = float(MipResolution.Y) / TileDimWithBorders.Y;
+		FIntPoint DimensionInTiles_PartialTiles
+			( FMath::CeilToInt(TileResFractionX)
+			, FMath::CeilToInt(TileResFractionY));
 
-			// Dimension of the texture in tiles excluding partial tiles.
-			FIntPoint DimensionInTiles_CompleteTiles
-				( FMath::FloorToInt(TileResFractionX)
-				, FMath::FloorToInt(TileResFractionY));
+		// Dimension of the texture in tiles excluding partial tiles.
+		FIntPoint DimensionInTiles_CompleteTiles
+			( FMath::FloorToInt(TileResFractionX)
+			, FMath::FloorToInt(TileResFractionY));
 
-			// Total number of tiles for this mip level.
-			int32 NumActualTiles = DimensionInTiles_PartialTiles.X * DimensionInTiles_PartialTiles.Y;
+		// Total number of tiles for this mip level.
+		int32 NumActualTiles = DimensionInTiles_PartialTiles.X * DimensionInTiles_PartialTiles.Y;
 
-			// Custom exr has one Vanilla tile per mip level that contains custom tiles.
-			if (bCustomExr)
-			{
-				OutNumTilesPerLevel.Add(1);
-			}
-			else
-			{
-				OutNumTilesPerLevel.Add(DimensionInTiles_PartialTiles.X * DimensionInTiles_PartialTiles.Y);
-			}
-
-			bool bHasPartialTiles = (DimensionInTiles_PartialTiles != DimensionInTiles_CompleteTiles) || CVarForceTileDescBufferExrGpuReader.GetValueOnAnyThread();
+		OutNumTilesPerLevel.Add(DimensionInTiles_PartialTiles.X * DimensionInTiles_PartialTiles.Y);
 
-			OutPartialTileInfo.Add({});
-			TArray<FTileDesc>& TileInfoList = OutPartialTileInfo[MipLevel];
-
-			if (!bCustomExr && !bHasPartialTiles)
-			{
-				continue;
-			}
-
-
-			// Resolution of the partial tile in the bottom right corner.
-			FIntPoint PartialTileResolution = FIntPoint(MipResolution.X % TileDimWithBorders.X, MipResolution.Y % TileDimWithBorders.Y);
+		const bool bHasPartialTiles = (DimensionInTiles_PartialTiles != DimensionInTiles_CompleteTiles) || CVarForceTileDescBufferExrGpuReader.GetValueOnAnyThread();
+		if (!bHasPartialTiles)
+		{
+			continue;
+		}
 
-			TArray<int64>& CurrentMipOffsets = OutCustomOffsets[MipLevel];
-			CurrentMipOffsets.SetNum(NumActualTiles);
+		OutPartialTileInfo.Add({});
+		TArray<FTileDesc>& TileInfoList = OutPartialTileInfo[MipLevel];
+		TileInfoList.SetNum(NumActualTiles);
 
-			if (bHasPartialTiles)
-			{
-				TileInfoList.SetNum(NumActualTiles);
-			}
+		// Resolution of the partial tile in the bottom right corner.
+		const FIntPoint PartialTileResolution = FIntPoint(MipResolution.X % TileDimWithBorders.X, MipResolution.Y % TileDimWithBorders.Y);
+		const int64 MipOffsetStart = CurrentPosition;
 
-			if (bCustomExr)
-			{
-				CurrentPosition += FExrReader::TILE_PADDING;
-			}
-
-			int64 MipOffsetStart = CurrentPosition;
+		for (int TileIndex = 0; TileIndex < NumActualTiles; TileIndex++)
+		{
+			FIntPoint TileDim(TileDimWithBorders);
 
-			for (int TileIndex = 0; TileIndex < NumActualTiles; TileIndex++)
+			if (DimensionInTiles_PartialTiles.X != DimensionInTiles_CompleteTiles.X)
 			{
-				FIntPoint TileDim(TileDimWithBorders);
-
-				if (DimensionInTiles_PartialTiles.X != DimensionInTiles_CompleteTiles.X)
+				// If true - this is a partial tile in X dimension.
+				if (((TileIndex + 1) % DimensionInTiles_PartialTiles.X) == 0)
 				{
-					// If true - this is a partial tile in X dimension.
-					if (((TileIndex + 1) % DimensionInTiles_PartialTiles.X) == 0)
-					{
-						TileDim.X = PartialTileResolution.X;
-					}
-				}
-
-				if (DimensionInTiles_PartialTiles.Y != DimensionInTiles_CompleteTiles.Y)
-				{
-					// If true - this is a partial tile in Y dimension.
-					if (TileIndex >= (DimensionInTiles_PartialTiles.X * DimensionInTiles_CompleteTiles.Y))
-					{
-						TileDim.Y = PartialTileResolution.Y;
-					}
+					TileDim.X = PartialTileResolution.X;
 				}
+			}
 
-				if (bHasPartialTiles)
+			if (DimensionInTiles_PartialTiles.Y != DimensionInTiles_CompleteTiles.Y)
+			{
+				// If true - this is a partial tile in Y dimension.
+				if (TileIndex >= (DimensionInTiles_PartialTiles.X * DimensionInTiles_CompleteTiles.Y))
 				{
-					TileInfoList[TileIndex] = { TileDim, (uint32)(CurrentPosition - MipOffsetStart) };
+					TileDim.Y = PartialTileResolution.Y;
 				}
+			}
 
-				CurrentMipOffsets[TileIndex] = CurrentPosition;
+			TileInfoList[TileIndex] = { TileDim, (uint32)(CurrentPosition - MipOffsetStart) };
 
-				// Tile offset.
-				CurrentPosition += TileDim.X * TileDim.Y * PixelSize;
+			// Tile offset.
+			CurrentPosition += TileDim.X * TileDim.Y * PixelSize;
 
-				// Vanilla exr has 20 byte padding at the beginning of each tile.
-				if (!bCustomExr)
-				{
-					CurrentPosition += FExrReader::TILE_PADDING;
-				}
-			}
+			// Vanilla exr has 20 byte padding at the beginning of each tile.
+			CurrentPosition += FExrReader::TILE_PADDING;
 		}
 	}
 }
 
-bool FExrReader::OpenExrAndPrepareForPixelReading(FString FilePath, const TArray<int32>& NumOffsetsPerLevel, TArray<TArray<int64>>&& CustomOffsets, bool bInCustomExr)
+bool FExrReader::OpenExrAndPrepareForPixelReading(FString FilePath, const TArray<int32>& NumOffsetsPerLevel)
 {
- 	bCustomExr = bInCustomExr;
 	if (FileHandle != nullptr)
 	{
 		UE_LOG(LogExrReaderGpu, Error, TEXT("The file has already been open for reading but never closed."));
@@ -413,12 +375,6 @@ bool FExrReader::OpenExrAndPrepareForPixelReading(FString FilePath, const TArray
 			OffsetsPerLevel.SetNum(NumOffsetsPerLevel[Level]);
 			ReadLineOrTileOffsets(FileHandle, LineOrder, OffsetsPerLevel);
 		}
-
-		if (bCustomExr)
-		{
-			PixelStartByteOffset = LineOrTileOffsetsPerLevel[0][0];
-			LineOrTileOffsetsPerLevel = MoveTemp(CustomOffsets);
-		}
 	}
 
 	fseek(FileHandle, 0, SEEK_END);
@@ -460,7 +416,7 @@ bool FExrReader::SeekTileWithinFile(const int32 StartTileIndex, const int32 Leve
 
 	int64 LineOffset = LineOrTileOffsetsPerLevel[Level][StartTileIndex];
 	OutBufferOffset = LineOffset - LineOrTileOffsetsPerLevel[Level][0];
-	return fseek(FileHandle, LineOffset + (bCustomExr ? PixelStartByteOffset : 0), SEEK_SET) == 0;
+	return fseek(FileHandle, LineOffset, SEEK_SET) == 0;
 }
 
 bool FExrReader::GetByteOffsetForTile(const int32 TileIndex, const int32 Level, int64& OutBufferOffset)
@@ -477,7 +433,7 @@ bool FExrReader::GetByteOffsetForTile(const int32 TileIndex, const int32 Level,
 			// If this is the last mip and the last tile that means that we don't have 
 			if (LineOrTileOffsetsPerLevel.Num() > Level + 1)
 			{
-				OutBufferOffset = LineOrTileOffsetsPerLevel[Level + 1][0] + (bCustomExr ? PixelStartByteOffset : 0);
+				OutBufferOffset = LineOrTileOffsetsPerLevel[Level + 1][0];
 			}
 			else
 			{
@@ -486,7 +442,7 @@ bool FExrReader::GetByteOffsetForTile(const int32 TileIndex, const int32 Level,
 		}
 		else
 		{
-			OutBufferOffset = LineOrTileOffsetsPerLevel[Level][TileIndex] + (bCustomExr ? PixelStartByteOffset : 0);
+			OutBufferOffset = LineOrTileOffsetsPerLevel[Level][TileIndex];
 		}
 		return true;
 	}
diff --git a/Engine/Plugins/Media/ImgMedia/Source/ExrReaderGpu/Public/ExrReaderGpu.h b/Engine/Plugins/Media/ImgMedia/Source/ExrReaderGpu/Public/ExrReaderGpu.h
--- a/Engine/Plugins/Media/ImgMedia/Source/ExrReaderGpu/Public/ExrReaderGpu.h
+++ b/Engine/Plugins/Media/ImgMedia/Source/ExrReaderGpu/Public/ExrReaderGpu.h
@@ -51,9 +51,7 @@ public:
 	FExrReader() 
 		: FileHandle(nullptr)
 		, FileLength(0)
-		, bCustomExr(false)
 		, LineOrTileOffsetsPerLevel()
-		, PixelStartByteOffset(0)
 	{};
 private:
 
@@ -80,20 +78,18 @@ public:
 	*/
 	static void CalculateTileOffsets
 		( TArray<int32>& OutNumTilesPerLevel
-		, TArray<TArray<int64>>& OutCustomOffsets
 		, TArray<TArray<FTileDesc>>& OutPartialTileInfo
 		, const FIntPoint& FullTextureResolution
 		, const FIntPoint& TileDimWithBorders
 		, int32 NumMipLevels
-		, int64 PixelSize
-		, bool bCustomExr);
+		, int64 PixelSize);
 
 	/** 
 	* This function is used to open file and keep a handle to it so the file can be read in chunks.
 	* Reading in chunks allows the process of reading to be canceled midway through reading. 
 	* This function reads and discards the header and keeps the pointers to the scanlines or tiles. 
 	*/
-	bool OpenExrAndPrepareForPixelReading(FString FilePath, const TArray<int32>& NumOffsetsPerLevel, TArray<TArray<int64>>&& CustomOffsets, bool bInCustomExr = false);
+	bool OpenExrAndPrepareForPixelReading(FString FilePath, const TArray<int32>& NumOffsetsPerLevel);
 
 	/** 
 	* Read a chunk of an Exr file previously open via OpenExrAndPrepareForPixelReading.
@@ -130,19 +126,9 @@ private:
 	*/
 	int64 FileLength;
 
-	/**
-	* Is this exr a custom exr with swizzled data.
-	*/
-	bool bCustomExr;
-
 	/**
 	* These are byte offsets pointing to scanlines or tiles from the begining of the file.
 	*/
 	TArray<TArray<int64>> LineOrTileOffsetsPerLevel;
-
-	/** 
-	* Byte offset of the first pixel in file. This is equal to the size of the header.
-	*/
-	int64 PixelStartByteOffset;
 };
 #endif
diff --git a/Engine/Plugins/Media/ImgMedia/Source/ExrReaderGpu/Public/ExrSwizzlingShader.h b/Engine/Plugins/Media/ImgMedia/Source/ExrReaderGpu/Public/ExrSwizzlingShader.h
--- a/Engine/Plugins/Media/ImgMedia/Source/ExrReaderGpu/Public/ExrSwizzlingShader.h
+++ b/Engine/Plugins/Media/ImgMedia/Source/ExrReaderGpu/Public/ExrSwizzlingShader.h
@@ -28,11 +28,10 @@ class FExrSwizzlePS : public FGlobalShader
 	SHADER_USE_PARAMETER_STRUCT(FExrSwizzlePS, FGlobalShader);
 
 	/** If the provided buffer is RGBA the shader would work slightly differently to RGB. */
-	class FRgbaSwizzle : SHADER_PERMUTATION_INT("NUM_CHANNELS", 4);
+	class FRgbaSwizzle : SHADER_PERMUTATION_INT("PERMUTATION_CHANNELS", 4);
 	class FRenderTiles : SHADER_PERMUTATION_BOOL("RENDER_TILES");
-	class FCustomExr : SHADER_PERMUTATION_BOOL("CUSTOM_EXR");
 	class FPartialTiles : SHADER_PERMUTATION_BOOL("PARTIAL_TILES");
-	using FPermutationDomain = TShaderPermutationDomain<FRgbaSwizzle, FRenderTiles, FCustomExr, FPartialTiles>;
+	using FPermutationDomain = TShaderPermutationDomain<FRgbaSwizzle, FRenderTiles, FPartialTiles>;
 
 	BEGIN_SHADER_PARAMETER_STRUCT(FParameters, )
 		SHADER_PARAMETER_SRV(StructuredBuffer<uint>, UnswizzledBuffer)
diff --git a/Engine/Plugins/Media/ImgMedia/Source/ImgMedia/Private/Loader/ImgMediaLoader.cpp b/Engine/Plugins/Media/ImgMedia/Source/ImgMedia/Private/Loader/ImgMediaLoader.cpp
--- a/Engine/Plugins/Media/ImgMedia/Source/ImgMedia/Private/Loader/ImgMediaLoader.cpp
+++ b/Engine/Plugins/Media/ImgMedia/Source/ImgMedia/Private/Loader/ImgMediaLoader.cpp
@@ -987,7 +987,7 @@ void FImgMediaLoader::WarmupSequence(const FImgMediaFrameInfo& InFirstFrameInfo,
 	// Giving our reader a chance to handle RAM allocation.
 	// Not all readers use this, only those that need to handle large files 
 	// or need to be as efficient as possible.
-	Reader->PreAllocateMemoryPool(NumFramesToLoad, InFirstFrameInfo, InFirstFrameInfo.FormatName == TEXT("EXR CUSTOM"));
+	Reader->PreAllocateMemoryPool(NumFramesToLoad, InFirstFrameInfo);
 
 	FScopeLock Lock(&CriticalSection);
 
diff --git a/Engine/Plugins/Media/ImgMedia/Source/ImgMedia/Private/Readers/ExrImgMediaReader.cpp b/Engine/Plugins/Media/ImgMedia/Source/ImgMedia/Private/Readers/ExrImgMediaReader.cpp
--- a/Engine/Plugins/Media/ImgMedia/Source/ImgMedia/Private/Readers/ExrImgMediaReader.cpp
+++ b/Engine/Plugins/Media/ImgMedia/Source/ImgMedia/Private/Readers/ExrImgMediaReader.cpp
@@ -30,9 +30,6 @@ static TAutoConsoleVariable<bool> CVarEnableUncompressedExrGpuReader(
 
 FExrImgMediaReader::FExrImgMediaReader(const TSharedRef<FImgMediaLoader, ESPMode::ThreadSafe>& InLoader)
 	: LoaderPtr(InLoader)
-	, bIsCustomFormat(false)
-	, bIsCustomFormatTiled(false)
-	, CustomFormatTileSize(EForceInit::ForceInitToZero)
 {
 	const UImgMediaSettings* Settings = GetDefault<UImgMediaSettings>();
 	
@@ -77,29 +74,26 @@ FExrImgMediaReader::EReadResult FExrImgMediaReader::ReadTiles
 		, FMath::CeilToInt(float(MipResolution.Y) / ConverterParams.TileDimWithBorders.Y));
 
 	TArray<int32> NumTilesPerLevel;
-	TArray<TArray<int64>> CustomOffsets;
 	int32 NumMipLevels = ConverterParams.bMipsInSeparateFiles ? 1 : ConverterParams.NumMipLevels;
 
 	FExrReader::CalculateTileOffsets(
 		NumTilesPerLevel,
-		CustomOffsets,
 		ConverterParams.TileInfoPerMipLevel,
 		ConverterParams.FullResolution,
 		ConverterParams.TileDimWithBorders,
 		NumMipLevels,
-		ConverterParams.PixelSize,
-		ConverterParams.bCustomExr);
+		ConverterParams.PixelSize);
 
 	{
-		TRACE_CPUPROFILER_EVENT_SCOPE_TEXT(*FString::Printf(TEXT("FExrImgMediaReader_ReadTilesCustom_OpenFile")));
-		if (!ChunkReader.OpenExrAndPrepareForPixelReading(ImagePath, NumTilesPerLevel, MoveTemp(CustomOffsets), ConverterParams.bCustomExr))
+		TRACE_CPUPROFILER_EVENT_SCOPE_TEXT(*FString::Printf(TEXT("FExrImgMediaReader_ReadTiles_OpenFile")));
+		if (!ChunkReader.OpenExrAndPrepareForPixelReading(ImagePath, NumTilesPerLevel))
 		{
 			return Fail;
 		}
 	}
 	{
 		int64 CurrentBufferPos = 0;
-		TRACE_CPUPROFILER_EVENT_SCOPE_TEXT(*FString::Printf(TEXT("FExrImgMediaReader_ReadTilesCustom_ReadTiles")));
+		TRACE_CPUPROFILER_EVENT_SCOPE_TEXT(*FString::Printf(TEXT("FExrImgMediaReader_ReadTiles_ReadTiles")));
 		for (const FIntRect& RawTileRegion : TileRegions)
 		{
 			// This clamp is to make sure that tile region is not out of bounds in case the region wasn't calculated incorrectly for some reason.
@@ -120,8 +114,7 @@ FExrImgMediaReader::EReadResult FExrImgMediaReader::ReadTiles
 					}
 				}
 
-				const uint16 Padding = ConverterParams.bCustomExr ? 0 : FExrReader::TILE_PADDING;
-				const int64 TileByteStride = ConverterParams.PixelSize * ConverterParams.TileDimWithBorders.X * ConverterParams.TileDimWithBorders.Y + Padding;
+				const uint16 Padding = FExrReader::TILE_PADDING;
 				const int StartTileIndex = TileRow * DimensionInTiles.X + TileRegion.Min.X;
 
 				bool bLastTile = TileRow + 1 == DimensionInTiles.Y && TileRegion.Max.X == DimensionInTiles.X;
@@ -263,39 +256,7 @@ bool FExrImgMediaReader::ReadFrame(int32 FrameId, const TMap<int32, FImgMediaTil
 				FString Image = Loader->GetImagePath(FrameId, CurrentMipLevel);
 				FString BaseImage;
 
-				if (FrameInfo.FormatName == TEXT("EXR CUSTOM"))
-				{
-#if defined(PLATFORM_WINDOWS) && PLATFORM_WINDOWS
-					TArray<FIntRect> TileRegions = CurrentTileSelection.GetVisibleRegions();
-					int32 PixelSize = sizeof(uint16) * FrameInfo.NumChannels;
-					FSampleConverterParameters ConverterParams;
-					ConverterParams.FrameInfo = FrameInfo;
-					ConverterParams.PixelSize = sizeof(uint16) * ConverterParams.FrameInfo.NumChannels;
-					ConverterParams.TileDimWithBorders = FrameInfo.TileDimensions + FrameInfo.TileBorder * 2;
-					ConverterParams.NumMipLevels = Loader->GetNumMipLevels();
-					ConverterParams.bCustomExr = FrameInfo.FormatName == TEXT("EXR CUSTOM");
-					ConverterParams.FrameId = FrameId;
-
-					TArray<UE::Math::TIntPoint<int64>> OutBufferRegionsToCopy;
-					EReadResult ReadResult = ReadTiles((uint16*)MipDataPtr, FrameInfo.UncompressedSize / MipLevelDiv, Image, TileRegions, ConverterParams, CurrentMipLevel, OutBufferRegionsToCopy);
-					if (ReadResult != Fail)
-					{
-						OutFrame->MipTilesPresent.Emplace(CurrentMipLevel, CurrentTileSelection);
-						
-						for (const FIntRect& Region : TileRegions)
-						{
-							OutFrame->NumTilesRead += Region.Area();
-						}
-					}
-					else
-					{
-						UE_LOG(LogImgMedia, Error, TEXT("Could not load %s"), *Image);
-					}
-#else
-					UE_LOG(LogImgMedia, Error, TEXT("Current platform doesn't support custom EXR file %s"), *Image);
-#endif
-				}
-				else if (bHasTiles)
+				if (bHasTiles)
 				{
 					UE_LOG(LogImgMedia, Error, TEXT("Non-GPU reader doesn't currently support natively-tiled EXR file %s"), *Image);
 				}
@@ -343,10 +304,6 @@ void FExrImgMediaReader::UncancelFrame(int32 FrameNumber)
 /** Gets reader type (GPU vs CPU) depending on size of EXR and its compression. */
 TSharedPtr<IImgMediaReader, ESPMode::ThreadSafe> FExrImgMediaReader::GetReader(const TSharedRef <FImgMediaLoader, ESPMode::ThreadSafe>& InLoader, FString FirstImageInSequencePath)
 {
-	bool bIsCustomFormat = false;
-	bool bIsOptimizedForGpu = false;
-	FIntPoint TileSize(EForceInit::ForceInitToZero);
-
 #if defined(PLATFORM_WINDOWS) && PLATFORM_WINDOWS
 	if (!FPaths::FileExists(FirstImageInSequencePath))
 	{
@@ -359,30 +316,21 @@ TSharedPtr<IImgMediaReader, ESPMode::ThreadSafe> FExrImgMediaReader::GetReader(c
 		return MakeShareable(new FExrImgMediaReader(InLoader));
 	}
 
-	bIsCustomFormat = Info.FormatName.Equals(TEXT("EXR CUSTOM"));
-	bIsOptimizedForGpu = Info.FormatName.Equals(TEXT("EXR GPU")) || bIsCustomFormat;
-	if (bIsCustomFormat)
-	{
-		TileSize = Info.TileDimensions;
-	}
-
 	// Check GetCompressionName of OpenExrWrapper for other compression names.
 	// todo: Add and test Vulkan support
 	if (GDynamicRHI && GDynamicRHI->GetInterfaceType() == ERHIInterfaceType::D3D12
 		&& Info.CompressionName == "Uncompressed" 
 		&& CVarEnableUncompressedExrGpuReader.GetValueOnAnyThread()
-		&& bIsOptimizedForGpu
+		&& Info.FormatName.Equals(TEXT("EXR GPU"))
 		)
 	{
 		TSharedRef<FExrImgMediaReaderGpu, ESPMode::ThreadSafe> GpuReader = 
 			MakeShared<FExrImgMediaReaderGpu, ESPMode::ThreadSafe>(InLoader);
-		GpuReader->SetCustomFormatInfo(bIsCustomFormat, TileSize);
 		return GpuReader;
 	}
 #endif
-	FExrImgMediaReader* Reader = new FExrImgMediaReader(InLoader);
-	Reader->SetCustomFormatInfo(bIsCustomFormat, TileSize);
-	return MakeShareable(Reader);
+
+	return MakeShareable(new FExrImgMediaReader(InLoader));
 }
 
 /* FExrImgMediaReader implementation
@@ -463,14 +411,6 @@ bool FExrImgMediaReader::GetInfo(const FString& FilePath, FImgMediaFrameInfo& Ou
 	return (OutInfo.UncompressedSize > 0) && (OutInfo.Dim.GetMin() > 0);
 }
 
-
-void FExrImgMediaReader::SetCustomFormatInfo(bool bInIsCustomFormat, const FIntPoint& InTileSize)
-{
-	bIsCustomFormat = bInIsCustomFormat;
-	CustomFormatTileSize = InTileSize;
-	bIsCustomFormatTiled = InTileSize.X != 0;
-}
-
 SIZE_T FExrImgMediaReader::GetMipBufferTotalSize(FIntPoint Dim, bool bInHasMips)
 {
 	SIZE_T Size = 0;
diff --git a/Engine/Plugins/Media/ImgMedia/Source/ImgMedia/Private/Readers/ExrImgMediaReader.h b/Engine/Plugins/Media/ImgMedia/Source/ImgMedia/Private/Readers/ExrImgMediaReader.h
--- a/Engine/Plugins/Media/ImgMedia/Source/ImgMedia/Private/Readers/ExrImgMediaReader.h
+++ b/Engine/Plugins/Media/ImgMedia/Source/ImgMedia/Private/Readers/ExrImgMediaReader.h
@@ -38,11 +38,6 @@ public:
 	/** Gets reader type (GPU vs CPU) depending on size of EXR and its compression. */
 	static TSharedPtr<IImgMediaReader, ESPMode::ThreadSafe> GetReader(const TSharedRef<FImgMediaLoader, ESPMode::ThreadSafe>& InLoader, FString FirstImageInSequencePath);
 
-	/** Query if our images are in our custom format. */
-	bool IsCustomFormat() const { return bIsCustomFormat; }
-	/** Gets the tile size of our custom format. */
-	FIntPoint GetCustomFormatTileSize() { return CustomFormatTileSize; }
-
 protected:
 	enum EReadResult
 	{
@@ -83,14 +78,6 @@ protected:
 		, const int32 CurrentMipLevel
 		, TArray<UE::Math::TIntPoint<int64>>& OutBufferRegionsToCopy);
 
-	/**
-	 * Sets parameters of our custom format images.
-	 *
-	 * @param bInIsCustomFormat		True if we are using custom format images.
-	 * @param TileSize				Size of our tiles. If (0, 0) then we are not using tiles.
-	 */
-	void SetCustomFormatInfo(bool bInIsCustomFormat, const FIntPoint& InTileSize);
-
 	/**
 	 * Gets the total size needed for all mips.
 	 *
@@ -104,12 +91,6 @@ protected:
 	FCriticalSection CanceledFramesCriticalSection;
 	TMap<int32, FRgbaInputFile*> PendingFrames;
 	TWeakPtr<FImgMediaLoader, ESPMode::ThreadSafe> LoaderPtr;
-	/** True if we are using our custom format. */
-	bool bIsCustomFormat;
-	/** True if our custom format images are tiled. */
-	bool bIsCustomFormatTiled;
-	/** Tile size of our custom format images. */
-	FIntPoint CustomFormatTileSize;
 };
 
 
diff --git a/Engine/Plugins/Media/ImgMedia/Source/ImgMedia/Private/Readers/ExrImgMediaReaderGpu.cpp b/Engine/Plugins/Media/ImgMedia/Source/ImgMedia/Private/Readers/ExrImgMediaReaderGpu.cpp
--- a/Engine/Plugins/Media/ImgMedia/Source/ImgMedia/Private/Readers/ExrImgMediaReaderGpu.cpp
+++ b/Engine/Plugins/Media/ImgMedia/Source/ImgMedia/Private/Readers/ExrImgMediaReaderGpu.cpp
@@ -126,7 +126,7 @@ FExrImgMediaReader::EReadResult FExrImgMediaReaderGpu::ReadMip
 	int MipLevelDiv = 1 << CurrentMipLevel;
 	FIntPoint CurrentMipDim = ConverterParams.FullResolution / MipLevelDiv;
 	const FImgMediaFrameInfo& FrameInfo = ConverterParams.FrameInfo;
-	const SIZE_T BufferSize = GetBufferSize(CurrentMipDim, FrameInfo.NumChannels, FrameInfo.bHasTiles, FrameInfo.NumTiles / MipLevelDiv, ConverterParams.bCustomExr);
+	const SIZE_T BufferSize = GetBufferSize(CurrentMipDim, FrameInfo.NumChannels, FrameInfo.bHasTiles, FrameInfo.NumTiles / MipLevelDiv);
 	
 	FStructuredBufferPoolItemSharedPtr BufferData = SampleConverter->GetOrCreateMipLevelBuffer(
 		CurrentMipLevel,
@@ -144,7 +144,7 @@ FExrImgMediaReader::EReadResult FExrImgMediaReaderGpu::ReadMip
 	{
 		TArray<UE::Math::TIntPoint<int64>> BufferRegionsToCopy;
 		// read frame data
-		if (FrameInfo.bHasTiles || ConverterParams.bCustomExr)
+		if (FrameInfo.bHasTiles)
 		{
 			TArray<FIntRect> TileRegionsToRead;
 			{
@@ -247,7 +247,6 @@ bool FExrImgMediaReaderGpu::ReadFrame(int32 FrameId, const TMap<int32, FImgMedia
 	ConverterParams.PixelSize = sizeof(uint16) * ConverterParams.FrameInfo.NumChannels;
 	ConverterParams.TileDimWithBorders = FrameInfo.TileDimensions + FrameInfo.TileBorder * 2;
 	ConverterParams.NumMipLevels = Loader->GetNumMipLevels();
-	ConverterParams.bCustomExr = FrameInfo.FormatName == TEXT("EXR CUSTOM");
 	ConverterParams.bMipsInSeparateFiles = Loader->MipsInSeparateFiles();
 
 	{
@@ -342,7 +341,7 @@ bool FExrImgMediaReaderGpu::ReadFrame(int32 FrameId, const TMap<int32, FImgMedia
 			for (const FIntRect& TileRegion : CurrentTileSelection.GetVisibleRegions())
 			{
 				FIntRect Viewport;
-				if (ConverterParams.FrameInfo.bHasTiles || ConverterParams.bCustomExr)
+				if (ConverterParams.FrameInfo.bHasTiles)
 				{
 					Viewport.Min = FIntPoint(ConverterParams.TileDimWithBorders.X * TileRegion.Min.X, ConverterParams.TileDimWithBorders.Y * TileRegion.Min.Y);
 					Viewport.Max = FIntPoint(ConverterParams.TileDimWithBorders.X * TileRegion.Max.X, ConverterParams.TileDimWithBorders.Y * TileRegion.Max.Y);
@@ -369,10 +368,10 @@ bool FExrImgMediaReaderGpu::ReadFrame(int32 FrameId, const TMap<int32, FImgMedia
 	return true;
 }
 
-void FExrImgMediaReaderGpu::PreAllocateMemoryPool(int32 NumFrames, const FImgMediaFrameInfo& FrameInfo, const bool bCustomExr)
+void FExrImgMediaReaderGpu::PreAllocateMemoryPool(int32 NumFrames, const FImgMediaFrameInfo& FrameInfo)
 {
 	TRACE_CPUPROFILER_EVENT_SCOPE_TEXT(*FString::Printf(TEXT("ExrReaderGpu.PreAllocateMemoryPool")));
-	SIZE_T AllocSize = GetBufferSize(FrameInfo.Dim, FrameInfo.NumChannels, FrameInfo.bHasTiles, FrameInfo.NumTiles, bCustomExr);
+	SIZE_T AllocSize = GetBufferSize(FrameInfo.Dim, FrameInfo.NumChannels, FrameInfo.bHasTiles, FrameInfo.NumTiles);
 	for (int32 FrameCacheNum = 0; FrameCacheNum < NumFrames; FrameCacheNum++)
 	{
 		AllocateGpuBufferFromPool(AllocSize);
@@ -397,7 +396,7 @@ FExrImgMediaReaderGpu::EReadResult FExrImgMediaReaderGpu::ReadInChunks(uint16* B
 	const int32 NumLevels = 1;
 	TArray<int32> NumTOffsetsPerLevel;
 	NumTOffsetsPerLevel.Add(Dim.Y);
-	if (!ChunkReader.OpenExrAndPrepareForPixelReading(ImagePath, NumTOffsetsPerLevel, TArray<TArray<int64>>()))
+	if (!ChunkReader.OpenExrAndPrepareForPixelReading(ImagePath, NumTOffsetsPerLevel))
 	{
 		return Fail;
 	}
@@ -437,9 +436,9 @@ FExrImgMediaReaderGpu::EReadResult FExrImgMediaReaderGpu::ReadInChunks(uint16* B
 	return bResult;
 }
 
-SIZE_T FExrImgMediaReaderGpu::GetBufferSize(const FIntPoint& Dim, int32 NumChannels, bool bHasTiles, const FIntPoint& TileNum, const bool bCustomExr)
+SIZE_T FExrImgMediaReaderGpu::GetBufferSize(const FIntPoint& Dim, int32 NumChannels, bool bHasTiles, const FIntPoint& TileNum)
 {
-	if (!bHasTiles && !bCustomExr)
+	if (!bHasTiles)
 	{
 		/** 
 		* Reading scanlines.
@@ -459,7 +458,7 @@ SIZE_T FExrImgMediaReaderGpu::GetBufferSize(const FIntPoint& Dim, int32 NumChann
 		* At the beginning of each tile there is 20 byte data that has information
 		* about number contents of tiles.
 		*/
-		const uint16 Padding = bCustomExr ? 0 : FExrReader::TILE_PADDING;
+		const uint16 Padding = FExrReader::TILE_PADDING;
 		SIZE_T BufferSize = Dim.X * Dim.Y * sizeof(uint16) * NumChannels + (TileNum.X * TileNum.Y) * Padding;
 		return BufferSize;
 	}
@@ -492,8 +491,7 @@ void FExrImgMediaReaderGpu::CreateSampleConverterCallback(TSharedPtr<FExrMediaTe
 
 			FExrSwizzlePS::FPermutationDomain PermutationVector;
 			PermutationVector.Set<FExrSwizzlePS::FRgbaSwizzle>(ConverterParams.FrameInfo.NumChannels - 1);
-			PermutationVector.Set<FExrSwizzlePS::FRenderTiles>(ConverterParams.FrameInfo.bHasTiles || ConverterParams.bCustomExr);
-			PermutationVector.Set<FExrSwizzlePS::FCustomExr>(ConverterParams.bCustomExr);
+			PermutationVector.Set<FExrSwizzlePS::FRenderTiles>(ConverterParams.FrameInfo.bHasTiles);
 			PermutationVector.Set<FExrSwizzlePS::FPartialTiles>(false);
 
 			FExrSwizzlePS::FParameters Parameters = FExrSwizzlePS::FParameters();
diff --git a/Engine/Plugins/Media/ImgMedia/Source/ImgMedia/Private/Readers/ExrImgMediaReaderGpu.h b/Engine/Plugins/Media/ImgMedia/Source/ImgMedia/Private/Readers/ExrImgMediaReaderGpu.h
--- a/Engine/Plugins/Media/ImgMedia/Source/ImgMedia/Private/Readers/ExrImgMediaReaderGpu.h
+++ b/Engine/Plugins/Media/ImgMedia/Source/ImgMedia/Private/Readers/ExrImgMediaReaderGpu.h
@@ -72,7 +72,7 @@ public:
 	/**
 	* For performance reasons we want to pre-allocate structured buffers to at least the number of concurrent frames.
 	*/
-	virtual void PreAllocateMemoryPool(int32 NumFrames, const FImgMediaFrameInfo& FrameInfo, const bool bCustomExr) override;
+	virtual void PreAllocateMemoryPool(int32 NumFrames, const FImgMediaFrameInfo& FrameInfo) override;
 
 protected:
 
@@ -88,7 +88,7 @@ protected:
 	 * @param Dim Dimensions of the image.
 	 * @param NumChannels Number of channels in the image.
 	 */
-	static SIZE_T GetBufferSize(const FIntPoint& Dim, int32 NumChannels, bool bHasTiles, const FIntPoint& TileNum, const bool bCustomExr);
+	static SIZE_T GetBufferSize(const FIntPoint& Dim, int32 NumChannels, bool bHasTiles, const FIntPoint& TileNum);
 
 	/**
 	* Creates Sample converter to be used by Media Texture Resource.
@@ -168,9 +168,6 @@ struct FSampleConverterParameters
 	/** Pixel stride in bytes. I.e. 2 bytes per pixel x 3 channels = 6. */
 	int32 PixelSize;
 
-	/** Identifies this exr as custom, therefore all data should be swizzled. */
-	bool bCustomExr;
-
 	/** Indicates if mips stored in individual files.*/
 	bool bMipsInSeparateFiles;
 
diff --git a/Engine/Plugins/Media/ImgMedia/Source/ImgMedia/Private/Readers/IImgMediaReader.h b/Engine/Plugins/Media/ImgMedia/Source/ImgMedia/Private/Readers/IImgMediaReader.h
--- a/Engine/Plugins/Media/ImgMedia/Source/ImgMedia/Private/Readers/IImgMediaReader.h
+++ b/Engine/Plugins/Media/ImgMedia/Source/ImgMedia/Private/Readers/IImgMediaReader.h
@@ -454,7 +454,7 @@ public:
 	 * maximum number of frames with as much efficiency as possible.
 	 *
 	 */
-	virtual void PreAllocateMemoryPool(int32 NumFrames, const FImgMediaFrameInfo& FrameInfo, const bool bCustomFormat) {};
+	virtual void PreAllocateMemoryPool(int32 NumFrames, const FImgMediaFrameInfo& FrameInfo) {};
 
 	/**
 	 * Used in case reader needs to do some processing once per frame.
