ImgMedia: Fixed mip estimation for overlay materials drawn at the post-upscale resolution.

As a solution, we create a virtual view struct that has the adjusted resolution.

diff --git a/Engine/Plugins/Media/ImgMedia/Source/ImgMedia/Private/ImgMediaSceneViewExtension.cpp b/Engine/Plugins/Media/ImgMedia/Source/ImgMedia/Private/ImgMediaSceneViewExtension.cpp
--- a/Engine/Plugins/Media/ImgMedia/Source/ImgMedia/Private/ImgMediaSceneViewExtension.cpp
+++ b/Engine/Plugins/Media/ImgMedia/Source/ImgMedia/Private/ImgMediaSceneViewExtension.cpp
@@ -149,6 +149,25 @@ void FImgMediaSceneViewExtension::CacheViewInfo(FSceneViewFamily& InViewFamily,
 	}
 #endif
 
+	/*
+	* As a mitigation for overlay materials drawn at the post-upscale resolution, we create a secondary (virtual) view with the adjusted resolution.
+	* We make the greedy assumption that users of this obscure console variable come from media plate composite logic.
+	* This ensures that mip estimation covers both regular and upscaled resolutions. It may not be the most efficient solve, but by far the safest and simplest.
+	*/
+	static const auto CVarTranslucencySPBasis = IConsoleManager::Get().FindConsoleVariable(TEXT("r.Translucency.ScreenPercentage.Basis"));
+	if (CVarTranslucencySPBasis && CVarTranslucencySPBasis->GetInt() == 1)
+	{
+		FImgMediaViewInfo PostUpscaleVirtualInfo = Info;
+
+		// Note: This is equivalent to FViewInfo::GetSecondaryViewRectSize
+		PostUpscaleVirtualInfo.ViewportRect = FIntRect(0, 0,
+			FMath::CeilToInt(View.UnscaledViewRect.Width() * InViewFamily.SecondaryViewFraction),
+			FMath::CeilToInt(View.UnscaledViewRect.Height() * InViewFamily.SecondaryViewFraction)
+		);
+
+		CachedViewInfos.Add(MoveTemp(PostUpscaleVirtualInfo));
+	}
+
 	CachedViewInfos.Add(MoveTemp(Info));
 }
 
