ImgMedia: Re-introduce custom exr detection to issue error log.

Internal-only debug tool as the custom EXR format was never publicly accessible.

#rnx

diff --git a/Engine/Plugins/Media/ImgMedia/Source/ImgMedia/Private/Readers/ExrImgMediaReader.cpp b/Engine/Plugins/Media/ImgMedia/Source/ImgMedia/Private/Readers/ExrImgMediaReader.cpp
--- a/Engine/Plugins/Media/ImgMedia/Source/ImgMedia/Private/Readers/ExrImgMediaReader.cpp
+++ b/Engine/Plugins/Media/ImgMedia/Source/ImgMedia/Private/Readers/ExrImgMediaReader.cpp
@@ -316,12 +316,20 @@ TSharedPtr<IImgMediaReader, ESPMode::ThreadSafe> FExrImgMediaReader::GetReader(c
 		return MakeShareable(new FExrImgMediaReader(InLoader));
 	}
 
+	const bool bIsDevelopmentCustomFormat = Info.FormatName.Equals(TEXT("EXR CUSTOM"));
+	if (bIsDevelopmentCustomFormat)
+	{
+		UE_LOG(LogImgMedia, Error, TEXT("Support for the (development-only) custom EXR format has been removed."));
+		return nullptr;
+	}
+
+	const bool bIsOptimizedForGpu = Info.FormatName.Equals(TEXT("EXR GPU"));
 	// Check GetCompressionName of OpenExrWrapper for other compression names.
 	// todo: Add and test Vulkan support
 	if (GDynamicRHI && GDynamicRHI->GetInterfaceType() == ERHIInterfaceType::D3D12
 		&& Info.CompressionName == "Uncompressed" 
 		&& CVarEnableUncompressedExrGpuReader.GetValueOnAnyThread()
-		&& Info.FormatName.Equals(TEXT("EXR GPU"))
+		&& bIsOptimizedForGpu
 		)
 	{
 		TSharedRef<FExrImgMediaReaderGpu, ESPMode::ThreadSafe> GpuReader = 
@@ -359,11 +367,6 @@ bool FExrImgMediaReader::GetInfo(const FString& FilePath, FImgMediaFrameInfo& Ou
 	if (bIsCustomFormat)
 	{
 		OutInfo.FormatName = TEXT("EXR CUSTOM");
-
-		// Get tile size.
-		HeaderReader.GetIntAttribute(IImgMediaModule::CustomFormatTileBorderAttributeName.Resolve().ToString(), OutInfo.TileBorder);
-		OutInfo.bHasTiles = HeaderReader.GetIntAttribute(IImgMediaModule::CustomFormatTileWidthAttributeName.Resolve().ToString(), OutInfo.TileDimensions.X);
-		OutInfo.bHasTiles = OutInfo.bHasTiles && HeaderReader.GetIntAttribute(IImgMediaModule::CustomFormatTileHeightAttributeName.Resolve().ToString(), OutInfo.TileDimensions.Y);
 	}
 	else
 	{
