Media: Improved support for mip level estimation at different view resolutions with compositing.

* Mip estimation is now done on separately labelled view infos, thus allowing computations at render and/or display resolutions individually.
* Added accurate mip estimation support for holdout composite by checking the holdout primitive flag on the media plate static mesh.

diff --git a/Engine/Plugins/Media/ImgMedia/Source/ImgMedia/Internal/Assets/ImgMediaMipMapInfo.h b/Engine/Plugins/Media/ImgMedia/Source/ImgMedia/Internal/Assets/ImgMediaMipMapInfo.h
--- a/Engine/Plugins/Media/ImgMedia/Source/ImgMedia/Internal/Assets/ImgMediaMipMapInfo.h
+++ b/Engine/Plugins/Media/ImgMedia/Source/ImgMedia/Internal/Assets/ImgMediaMipMapInfo.h
@@ -214,6 +214,9 @@ public:
 	/** Get the mip level to upscale (-1 if disabled). */
 	virtual int32 GetMipLevelToUpscale() const;
 
+	/** Get the target view resolution. */
+	EMediaTextureTargetViewResolution GetTargetViewResolutionMask() const;
+
 protected:
 
 	/** Mesh component onto which the media is displayed. */
@@ -315,6 +318,9 @@ protected:
 	/** Info for each view, used in mipmap calculations. */
 	TArray<FImgMediaViewInfo> ViewInfos;
 
+	/** Info for each view at display resolution, used in mipmap calculations. */
+	TArray<FImgMediaViewInfo> DisplayResolutionViewInfos;
+
 	/** Size, tiling and mips sequence information. */
 	FSequenceInfo SequenceInfo;
 	
diff --git a/Engine/Plugins/Media/ImgMedia/Source/ImgMedia/Private/Assets/ImgMediaMipMapInfo.cpp b/Engine/Plugins/Media/ImgMedia/Source/ImgMedia/Private/Assets/ImgMediaMipMapInfo.cpp
--- a/Engine/Plugins/Media/ImgMedia/Source/ImgMedia/Private/Assets/ImgMediaMipMapInfo.cpp
+++ b/Engine/Plugins/Media/ImgMedia/Source/ImgMedia/Private/Assets/ImgMediaMipMapInfo.cpp
@@ -295,6 +295,17 @@ int32 FImgMediaMipMapObjectInfo::GetMipLevelToUpscale() const
 	return -1;
 }
 
+EMediaTextureTargetViewResolution FImgMediaMipMapObjectInfo::GetTargetViewResolutionMask() const
+{
+	TSharedPtr<FMediaTextureTrackerObject> PinnedTracker = Tracker.Pin();
+	if (PinnedTracker.IsValid())
+	{
+		return PinnedTracker->TargetViewResolutionMask;
+	}
+
+	return EMediaTextureTargetViewResolution::RenderResolution;
+}
+
 void FImgMediaMipMapObjectInfo::CalculateVisibleTiles(const TArray<FImgMediaViewInfo>& InViewInfos, const FSequenceInfo& InSequenceInfo, TMap<int32, FImgMediaTileSelection>& VisibleTiles) const
 {
 	UMeshComponent* Mesh = MeshComponent.Get();
@@ -405,6 +416,11 @@ namespace {
 				return;
 			}
 
+			if (InViewInfos.IsEmpty())
+			{
+				return;
+			}
+
 			const float MipMapBias = ObjectInfo->MipMapLODBias;
 			const float MipMapLevelPadding = FMath::Max(CVarImgMediaMipLevelPadding.GetValueOnAnyThread(), 0.0f);
 
@@ -648,6 +664,11 @@ namespace {
 				return;
 			}
 
+			if (InViewInfos.IsEmpty())
+			{
+				return;
+			}
+
 			const FVector2f MeshRange = FVector2f(ObjectInfo->MeshRange);
 			const float MipMapBias = ObjectInfo->MipMapLODBias;
 			const float MipMapLevelPadding = FMath::Max(CVarImgMediaMipLevelPadding.GetValueOnAnyThread(), 0.0f);
@@ -1052,7 +1073,20 @@ TMap<int32, FImgMediaTileSelection> FImgMediaMipMapInfo::GetVisibleTiles()
 		// Loop over all objects.
 		for (FImgMediaMipMapObjectInfo* ObjectInfo : Objects)
 		{
-			ObjectInfo->CalculateVisibleTiles(ViewInfos, SequenceInfo, CachedVisibleTiles);
+			const EMediaTextureTargetViewResolution Mask = ObjectInfo->GetTargetViewResolutionMask();
+
+			// If display & render resolutions are identical, display resolution view infos will always be empty.
+			const bool bForceRenderResolution = DisplayResolutionViewInfos.IsEmpty();
+
+			if (EnumHasAllFlags(Mask, EMediaTextureTargetViewResolution::RenderResolution) || bForceRenderResolution)
+			{
+				ObjectInfo->CalculateVisibleTiles(ViewInfos, SequenceInfo, CachedVisibleTiles);
+			}
+
+			if (EnumHasAllFlags(Mask, EMediaTextureTargetViewResolution::DisplayResolution) && !bForceRenderResolution)
+			{
+				ObjectInfo->CalculateVisibleTiles(DisplayResolutionViewInfos, SequenceInfo, CachedVisibleTiles);
+			}
 		}
 
 		// Mark cache as valid.
@@ -1172,6 +1206,7 @@ void FImgMediaMipMapInfo::OnEndFrame()
 		if (SVE.IsValid())
 		{
 			ViewInfos = SVE->GetViewInfos();
+			DisplayResolutionViewInfos = SVE->GetDisplayResolutionViewInfos();
 		}
 	}
 }
diff --git a/Engine/Plugins/Media/ImgMedia/Source/ImgMedia/Private/ImgMediaSceneViewExtension.cpp b/Engine/Plugins/Media/ImgMedia/Source/ImgMedia/Private/ImgMediaSceneViewExtension.cpp
--- a/Engine/Plugins/Media/ImgMedia/Source/ImgMedia/Private/ImgMediaSceneViewExtension.cpp
+++ b/Engine/Plugins/Media/ImgMedia/Source/ImgMedia/Private/ImgMediaSceneViewExtension.cpp
@@ -37,6 +37,7 @@ static TAutoConsoleVariable<bool> CVarImgMediaProcessTilesInnerOnly(
 FImgMediaSceneViewExtension::FImgMediaSceneViewExtension(const FAutoRegister& AutoReg)
 	: FSceneViewExtensionBase(AutoReg)
 	, CachedViewInfos()
+	, DisplayResolutionCachedViewInfos()
 {
 	OnBeginFrameDelegate = FCoreDelegates::OnBeginFrame.AddRaw(this, &FImgMediaSceneViewExtension::ResetViewInfoCache);
 }
@@ -89,6 +90,7 @@ void FImgMediaSceneViewExtension::CacheViewInfo(FSceneViewFamily& InViewFamily,
 	Info.Location = View.ViewMatrices.GetViewOrigin();
 	Info.ViewDirection = View.GetViewDirection();
 	Info.ViewProjectionMatrix = View.ViewMatrices.GetViewProjectionMatrix();
+	Info.ViewportRect = View.UnconstrainedViewRect.Scale(ResolutionFraction);
 
 	if (FMath::IsNearlyEqual(FieldOfViewMultiplier, 1.0f))
 	{
@@ -107,8 +109,6 @@ void FImgMediaSceneViewExtension::CacheViewInfo(FSceneViewFamily& InViewFamily,
 		Info.OverscanViewProjectionMatrix = View.ViewMatrices.GetViewMatrix() * AdjustedProjectionMatrix;
 	}
 
-	Info.ViewportRect = View.UnconstrainedViewRect.Scale(ResolutionFraction);
-
 	// We store hidden or show-only ids to later avoid needless calculations when objects are not in view.
 	if (View.ShowOnlyPrimitives.IsSet())
 	{
@@ -149,23 +149,18 @@ void FImgMediaSceneViewExtension::CacheViewInfo(FSceneViewFamily& InViewFamily,
 	}
 #endif
 
-	/*
-	* As a mitigation for overlay materials drawn at the post-upscale resolution, we create a secondary (virtual) view with the adjusted resolution.
-	* We make the greedy assumption that users of this obscure console variable come from media plate composite logic.
-	* This ensures that mip estimation covers both regular and upscaled resolutions. It may not be the most efficient solve, but by far the safest and simplest.
-	*/
-	static const auto CVarTranslucencySPBasis = IConsoleManager::Get().FindConsoleVariable(TEXT("r.Translucency.ScreenPercentage.Basis"));
-	if (CVarTranslucencySPBasis && CVarTranslucencySPBasis->GetInt() == 1)
+	// We cache the display resolution view info in case it's needed for compositing applications
+	const bool bIsDisplayResolutionDifferent = !FMath::IsNearlyEqual(ResolutionFraction, InViewFamily.SecondaryViewFraction);
+	if (bIsDisplayResolutionDifferent)
 	{
 		FImgMediaViewInfo PostUpscaleVirtualInfo = Info;
-
-		// Note: This is equivalent to FViewInfo::GetSecondaryViewRectSize
+		PostUpscaleVirtualInfo.MaterialTextureMipBias = 0.0f;
 		PostUpscaleVirtualInfo.ViewportRect = FIntRect(0, 0,
-			FMath::CeilToInt(View.UnscaledViewRect.Width() * InViewFamily.SecondaryViewFraction),
-			FMath::CeilToInt(View.UnscaledViewRect.Height() * InViewFamily.SecondaryViewFraction)
+			FMath::CeilToInt(View.UnconstrainedViewRect.Width() * InViewFamily.SecondaryViewFraction),
+			FMath::CeilToInt(View.UnconstrainedViewRect.Height() * InViewFamily.SecondaryViewFraction)
 		);
 
-		CachedViewInfos.Add(MoveTemp(PostUpscaleVirtualInfo));
+		DisplayResolutionCachedViewInfos.Add(MoveTemp(PostUpscaleVirtualInfo));
 	}
 
 	CachedViewInfos.Add(MoveTemp(Info));
@@ -174,6 +169,7 @@ void FImgMediaSceneViewExtension::CacheViewInfo(FSceneViewFamily& InViewFamily,
 void FImgMediaSceneViewExtension::ResetViewInfoCache()
 {
 	CachedViewInfos.Reset();
+	DisplayResolutionCachedViewInfos.Reset();
 }
 
 #undef LOCTEXT_NAMESPACE
diff --git a/Engine/Plugins/Media/ImgMedia/Source/ImgMedia/Public/ImgMediaSceneViewExtension.h b/Engine/Plugins/Media/ImgMedia/Source/ImgMedia/Public/ImgMediaSceneViewExtension.h
--- a/Engine/Plugins/Media/ImgMedia/Source/ImgMedia/Public/ImgMediaSceneViewExtension.h
+++ b/Engine/Plugins/Media/ImgMedia/Source/ImgMedia/Public/ImgMediaSceneViewExtension.h
@@ -46,6 +46,14 @@ public:
 	 */
 	IMGMEDIA_API const TArray<FImgMediaViewInfo>& GetViewInfos() const { return CachedViewInfos; };
 
+	/**
+	 * Get the cached camera information array at display resolution for compositing, updated on the game thread by BeginRenderViewFamily.
+	 * Will remain empty if the render resolution matches the display resolution.
+	 *
+	 * @return Array of info on each camera, at display resolution.
+	 */
+	IMGMEDIA_API const TArray<FImgMediaViewInfo>& GetDisplayResolutionViewInfos() const { return DisplayResolutionCachedViewInfos; };
+
 	void SetupViewFamily(FSceneViewFamily& InViewFamily) override {}
 	void SetupView(FSceneViewFamily& InViewFamily, FSceneView& InView) override {}
 	void BeginRenderViewFamily(FSceneViewFamily& InViewFamily) override;
@@ -61,6 +69,9 @@ private:
 	/** Array of info on each camera used for mipmap calculations. */
 	TArray<FImgMediaViewInfo> CachedViewInfos;
 
+	/** Array of info on each camera used for mipmap calculations, at display resolution. */
+	TArray<FImgMediaViewInfo> DisplayResolutionCachedViewInfos;
+
 	/** FCoreDelegates::OnBeginFrame delegate. */
 	FDelegateHandle OnBeginFrameDelegate;
 };
diff --git a/Engine/Plugins/Media/MediaPlate/Source/MediaPlate/Private/MediaPlateComponent.cpp b/Engine/Plugins/Media/MediaPlate/Source/MediaPlate/Private/MediaPlateComponent.cpp
--- a/Engine/Plugins/Media/MediaPlate/Source/MediaPlate/Private/MediaPlateComponent.cpp
+++ b/Engine/Plugins/Media/MediaPlate/Source/MediaPlate/Private/MediaPlateComponent.cpp
@@ -534,6 +534,35 @@ void UMediaPlateComponent::RegisterWithMediaTextureTracker()
 	MediaTextureTrackerObject->MipLevelToUpscale = bEnableMipMapUpscaling ? MipLevelToUpscale : -1;
 	MediaTextureTrackerObject->bAdaptivePoleMipUpscaling = bAdaptivePoleMipUpscaling;
 
+	// Specify view target resolution for any subsequent mip-level estimation, defaulting to render resolution.
+	EMediaTextureTargetViewResolution TargetViewResolutionMask = EMediaTextureTargetViewResolution::RenderResolution;
+	if (AMediaPlate* MediaPlate = GetOwner<AMediaPlate>())
+	{
+		// If holdout is enabled, we assume the holdout compositing is active and only use the display resolution.
+		if (MediaPlate->StaticMeshComponent && MediaPlate->StaticMeshComponent->bHoldout)
+		{
+			TargetViewResolutionMask = EMediaTextureTargetViewResolution::DisplayResolution;
+		}
+		else
+		{
+			// First we check if overlay rendering is done at display resolution
+			static const auto CVarTranslucencySPBasis = IConsoleManager::Get().FindConsoleVariable(TEXT("r.Translucency.ScreenPercentage.Basis"));
+			const bool bDisplayResolutionOverlay = CVarTranslucencySPBasis && CVarTranslucencySPBasis->GetInt() == 1;
+
+			if (bDisplayResolutionOverlay && MediaPlate->GetCurrentOverlayMaterial() != nullptr)
+			{
+				EnumAddFlags(TargetViewResolutionMask, EMediaTextureTargetViewResolution::DisplayResolution);
+
+				// If we only have an overlay material, we can safely remove the render resolution target
+				if (MediaPlate->GetCurrentMaterial() == nullptr)
+				{
+					EnumRemoveFlags(TargetViewResolutionMask, EMediaTextureTargetViewResolution::RenderResolution);
+				}
+			}
+		}
+	}
+	MediaTextureTrackerObject->TargetViewResolutionMask = TargetViewResolutionMask;
+
 	// Add our textures.
 	FMediaTextureTracker& MediaTextureTracker = FMediaTextureTracker::Get();
 	for (UMediaTexture* MediaTexture : MediaTextures)
diff --git a/Engine/Source/Runtime/MediaAssets/Public/MediaTextureTracker.h b/Engine/Source/Runtime/MediaAssets/Public/MediaTextureTracker.h
--- a/Engine/Source/Runtime/MediaAssets/Public/MediaTextureTracker.h
+++ b/Engine/Source/Runtime/MediaAssets/Public/MediaTextureTracker.h
@@ -21,6 +21,15 @@ enum class EMediaTextureVisibleMipsTiles : uint8
 	Sphere
 };
 
+enum class EMediaTextureTargetViewResolution : uint8
+{
+	None = 0,
+	RenderResolution = 1 << 0,
+	DisplayResolution = 1 << 1,
+	All = RenderResolution | DisplayResolution
+};
+ENUM_CLASS_FLAGS(EMediaTextureTargetViewResolution);
+
 /** Holds info on a single object. */
 struct FMediaTextureTrackerObject
 {
@@ -36,6 +45,8 @@ struct FMediaTextureTrackerObject
 	int32 MipLevelToUpscale = -1;
 	/* When this is enabled Media Plate will automatically try to figure out the best mip to upscale to reduce the load. */
 	bool bAdaptivePoleMipUpscaling = false;
+	/* Bitwise mask for the target displayed resolution of the media texture (used for mip level estimation). */
+	EMediaTextureTargetViewResolution TargetViewResolutionMask = EMediaTextureTargetViewResolution::RenderResolution;
 };
 
 /**
