[Insights] Fixed usage of TRACE_CPUPROFILER_EVENT_SCOPE / *_STR / *_TEXT in engine plugins.

#rnx

diff --git a/Engine/Plugins/Media/ImgMedia/Source/ExrReaderGpu/Private/ExrReaderGpu.cpp b/Engine/Plugins/Media/ImgMedia/Source/ExrReaderGpu/Private/ExrReaderGpu.cpp
--- a/Engine/Plugins/Media/ImgMedia/Source/ExrReaderGpu/Private/ExrReaderGpu.cpp
+++ b/Engine/Plugins/Media/ImgMedia/Source/ExrReaderGpu/Private/ExrReaderGpu.cpp
@@ -267,7 +267,7 @@ void FExrReader::CalculateTileOffsets
 	, int32 NumMipLevels
 	, int64 PixelSize)
 {
-	TRACE_CPUPROFILER_EVENT_SCOPE_TEXT(*FString::Printf(TEXT("CalculateTileOffsets")));
+	TRACE_CPUPROFILER_EVENT_SCOPE(CalculateTileOffsets);
 
 	int64 CurrentPosition = 0;
 	for (int32 MipLevel = 0; MipLevel < NumMipLevels; MipLevel++)
diff --git a/Engine/Plugins/Media/ImgMedia/Source/ImgMedia/Private/Readers/ExrImgMediaReader.cpp b/Engine/Plugins/Media/ImgMedia/Source/ImgMedia/Private/Readers/ExrImgMediaReader.cpp
--- a/Engine/Plugins/Media/ImgMedia/Source/ImgMedia/Private/Readers/ExrImgMediaReader.cpp
+++ b/Engine/Plugins/Media/ImgMedia/Source/ImgMedia/Private/Readers/ExrImgMediaReader.cpp
@@ -85,7 +85,7 @@ FExrImgMediaReader::EReadResult FExrImgMediaReader::ReadTiles
 		ConverterParams.PixelSize);
 
 	{
-		TRACE_CPUPROFILER_EVENT_SCOPE_TEXT(*FString::Printf(TEXT("FExrImgMediaReader_ReadTiles_OpenFile")));
+		TRACE_CPUPROFILER_EVENT_SCOPE_STR("FExrImgMediaReader_ReadTiles_OpenFile");
 		if (!ChunkReader.OpenExrAndPrepareForPixelReading(ImagePath, NumTilesPerLevel))
 		{
 			return Fail;
@@ -93,7 +93,7 @@ FExrImgMediaReader::EReadResult FExrImgMediaReader::ReadTiles
 	}
 	{
 		int64 CurrentBufferPos = 0;
-		TRACE_CPUPROFILER_EVENT_SCOPE_TEXT(*FString::Printf(TEXT("FExrImgMediaReader_ReadTiles_ReadTiles")));
+		TRACE_CPUPROFILER_EVENT_SCOPE_STR("FExrImgMediaReader_ReadTiles_ReadTiles");
 		for (const FIntRect& RawTileRegion : TileRegions)
 		{
 			// This clamp is to make sure that tile region is not out of bounds in case the region wasn't calculated incorrectly for some reason.
diff --git a/Engine/Plugins/Media/ImgMedia/Source/ImgMedia/Private/Readers/ExrImgMediaReaderGpu.cpp b/Engine/Plugins/Media/ImgMedia/Source/ImgMedia/Private/Readers/ExrImgMediaReaderGpu.cpp
--- a/Engine/Plugins/Media/ImgMedia/Source/ImgMedia/Private/Readers/ExrImgMediaReaderGpu.cpp
+++ b/Engine/Plugins/Media/ImgMedia/Source/ImgMedia/Private/Readers/ExrImgMediaReaderGpu.cpp
@@ -370,7 +370,7 @@ bool FExrImgMediaReaderGpu::ReadFrame(int32 FrameId, const TMap<int32, FImgMedia
 
 void FExrImgMediaReaderGpu::PreAllocateMemoryPool(int32 NumFrames, const FImgMediaFrameInfo& FrameInfo)
 {
-	TRACE_CPUPROFILER_EVENT_SCOPE_TEXT(*FString::Printf(TEXT("ExrReaderGpu.PreAllocateMemoryPool")));
+	TRACE_CPUPROFILER_EVENT_SCOPE_STR("ExrReaderGpu.PreAllocateMemoryPool");
 	SIZE_T AllocSize = GetBufferSize(FrameInfo.Dim, FrameInfo.NumChannels, FrameInfo.bHasTiles, FrameInfo.NumTiles);
 	for (int32 FrameCacheNum = 0; FrameCacheNum < NumFrames; FrameCacheNum++)
 	{
@@ -506,7 +506,7 @@ void FExrImgMediaReaderGpu::CreateSampleConverterCallback(TSharedPtr<FExrMediaTe
 			if (ConverterParams.FrameInfo.bHasTiles &&
 				(ConverterParams.TileInfoPerMipLevel.Num() > SampleMipLevel && ConverterParams.TileInfoPerMipLevel[SampleMipLevel].Num() > 0))
 			{
-				TRACE_CPUPROFILER_EVENT_SCOPE_TEXT(*FString::Printf(TEXT("ExrReaderGpu.TileDesc")));
+				TRACE_CPUPROFILER_EVENT_SCOPE_STR("ExrReaderGpu.TileDesc");
 
 				FBufferRHIRef BufferRef;
 				FRHIResourceCreateInfo CreateInfo(TEXT("FExrImgMediaReaderGpu_TileDesc"));
@@ -608,7 +608,7 @@ void FExrImgMediaReaderGpu::CreateSampleConverterCallback(TSharedPtr<FExrMediaTe
 
 FStructuredBufferPoolItemSharedPtr FExrImgMediaReaderGpu::AllocateGpuBufferFromPool(uint32 AllocSize)
 {
-	TRACE_CPUPROFILER_EVENT_SCOPE_TEXT(*FString::Printf(TEXT("ExrReaderGpu.AllocBuffer")));
+	//TRACE_CPUPROFILER_EVENT_SCOPE_STR("ExrReaderGpu.AllocBuffer");
 	TRACE_CPUPROFILER_EVENT_SCOPE_TEXT(*FString::Printf(TEXT("ExrReaderGpu.AllocBuffer %d"), AllocSize));
 	TWeakPtr<FExrImgMediaReaderGpu, ESPMode::ThreadSafe> WeakReaderPtr = AsWeak();
 
@@ -648,7 +648,7 @@ FStructuredBufferPoolItemSharedPtr FExrImgMediaReaderGpu::AllocateGpuBufferFromP
 		// Allocate and unlock the structured buffer on render thread.
 		ENQUEUE_RENDER_COMMAND(CreatePooledBuffer)([AllocatedBuffer, AllocSize](FRHICommandListImmediate& RHICmdList)
 			{
-				TRACE_CPUPROFILER_EVENT_SCOPE_TEXT(*FString::Printf(TEXT("ExrReaderGpu.AllocBuffer_RenderThread")));
+				//TRACE_CPUPROFILER_EVENT_SCOPE_STR("ExrReaderGpu.AllocBuffer_RenderThread");
 				TRACE_CPUPROFILER_EVENT_SCOPE_TEXT(*FString::Printf(TEXT("ExrReaderGpu.AllocBuffer_RenderThread %d"), AllocSize));
 
 				SCOPED_GPU_STAT(RHICmdList, ExrImgMediaReaderGpu_AllocateBuffer);
@@ -681,7 +681,7 @@ FStructuredBufferPoolItemSharedPtr FExrImgMediaReaderGpu::AllocateGpuBufferFromP
 void FExrImgMediaReaderGpu::ReturnGpuBufferToPool(uint32 AllocSize, FStructuredBufferPoolItem* Buffer)
 {
 	FScopeLock ScopeLock(&MemoryPoolCriticalSection);
-	TRACE_CPUPROFILER_EVENT_SCOPE_TEXT(*FString::Printf(TEXT("ExrReaderGpu.ReturnPoolItem")));
+	TRACE_CPUPROFILER_EVENT_SCOPE_STR("ExrReaderGpu.ReturnPoolItem");
 	MemoryPool.Add(AllocSize, Buffer);
 }
 
@@ -709,7 +709,7 @@ FStructuredBufferPoolItem::FStructuredBufferPoolItem()
 
 FStructuredBufferPoolItem::~FStructuredBufferPoolItem()
 {
-	TRACE_CPUPROFILER_EVENT_SCOPE_TEXT(*FString::Printf(TEXT("ExrReaderGpu.ReleasePoolItem")));
+	TRACE_CPUPROFILER_EVENT_SCOPE_STR("ExrReaderGpu.ReleasePoolItem");
 	FRHICommandListImmediate::Get().UnlockBuffer(UploadBufferRef);
 	UploadBufferMapped = nullptr;
 
diff --git a/Engine/Plugins/Runtime/InstancedActors/Source/InstancedActors/Private/InstancedActorsManager.cpp b/Engine/Plugins/Runtime/InstancedActors/Source/InstancedActors/Private/InstancedActorsManager.cpp
--- a/Engine/Plugins/Runtime/InstancedActors/Source/InstancedActors/Private/InstancedActorsManager.cpp
+++ b/Engine/Plugins/Runtime/InstancedActors/Source/InstancedActors/Private/InstancedActorsManager.cpp
@@ -180,7 +180,7 @@ void AInstancedActorsManager::BeginPlay()
 
 	Super::BeginPlay();
 
-	TRACE_CPUPROFILER_EVENT_SCOPE("AInstancedActorsManager::BeginPlay");
+	TRACE_CPUPROFILER_EVENT_SCOPE(AInstancedActorsManager::BeginPlay);
 
 	UE_LOG(LogInstancedActors, Verbose, TEXT("%s (%d instances) BeginPlay"), *GetPathName(), GetNumValidInstances());
 	UMassEntitySubsystem* EntitySubsystem = World->GetSubsystem<UMassEntitySubsystem>();
diff --git a/Engine/Plugins/Runtime/InstancedActors/Source/InstancedActors/Private/InstancedActorsModifierVolumeComponent.cpp b/Engine/Plugins/Runtime/InstancedActors/Source/InstancedActors/Private/InstancedActorsModifierVolumeComponent.cpp
--- a/Engine/Plugins/Runtime/InstancedActors/Source/InstancedActors/Private/InstancedActorsModifierVolumeComponent.cpp
+++ b/Engine/Plugins/Runtime/InstancedActors/Source/InstancedActors/Private/InstancedActorsModifierVolumeComponent.cpp
@@ -30,7 +30,7 @@ void UInstancedActorsModifierVolumeComponent::BeginPlay()
 {
 	Super::BeginPlay();
 
-	TRACE_CPUPROFILER_EVENT_SCOPE("UInstancedActorsModifierVolumeComponent BeginPlay");
+	TRACE_CPUPROFILER_EVENT_SCOPE_STR("UInstancedActorsModifierVolumeComponent BeginPlay");
 
 	// Make sure bounds are up to date e.g: on clients we will have only just received Extent replication
 	// but not updated bounds via any of the normal component init mechanisms.
@@ -73,7 +73,7 @@ void UInstancedActorsModifierVolumeComponent::EndPlay(EEndPlayReason::Type Reaso
 {
 	Super::EndPlay(Reason);
 
-	TRACE_CPUPROFILER_EVENT_SCOPE("UInstancedActorsModifierVolumeComponent EndPlay");
+	TRACE_CPUPROFILER_EVENT_SCOPE_STR("UInstancedActorsModifierVolumeComponent EndPlay");
 
 	UInstancedActorsSubsystem* InstancedActorSubsystem = UInstancedActorsSubsystem::Get(this);
 	if (ensure(InstancedActorSubsystem))
@@ -125,7 +125,7 @@ void UInstancedActorsModifierVolumeComponent::OnRemovedFromManager(AInstancedAct
 
 bool UInstancedActorsModifierVolumeComponent::TryRunPendingModifiers(AInstancedActorsManager& Manager, TBitArray<>& InOutPendingModifiers)
 {
-	TRACE_CPUPROFILER_EVENT_SCOPE("UInstancedActorsModifierVolumeComponent ModifyInstances");
+	TRACE_CPUPROFILER_EVENT_SCOPE_STR("UInstancedActorsModifierVolumeComponent ModifyInstances");
 
 #if WITH_EDITOR
 	// Need to do this or the PIE experience is not a valid representation of the cooked game.
