[Media] Adding Timecode to ImgMediaSource

diff --git a/Engine/Plugins/Media/ImgMedia/Source/ImgMedia/Private/Assets/ImgMediaSource.cpp b/Engine/Plugins/Media/ImgMedia/Source/ImgMedia/Private/Assets/ImgMediaSource.cpp
--- a/Engine/Plugins/Media/ImgMedia/Source/ImgMedia/Private/Assets/ImgMediaSource.cpp
+++ b/Engine/Plugins/Media/ImgMedia/Source/ImgMedia/Private/Assets/ImgMediaSource.cpp
@@ -8,6 +8,7 @@
 #include "ImgMediaPrivate.h"
 
 #include "HAL/FileManager.h"
+#include "MediaPlayer.h"
 #include "Misc/Paths.h"
 #include "UObject/UE5MainStreamObjectVersion.h"
 
@@ -160,6 +161,11 @@ FString UImgMediaSource::GetMediaOption(const FName& Key, const FString& Default
 		return ProxyOverride;
 	}
 
+	if (Key == UMediaPlayer::MediaInfoNameStartTimecodeValue.Resolve())
+	{
+		return StartTimecode.ToString();
+	}
+
 	return Super::GetMediaOption(Key, DefaultValue);
 }
 
diff --git a/Engine/Plugins/Media/ImgMedia/Source/ImgMedia/Public/ImgMediaSource.h b/Engine/Plugins/Media/ImgMedia/Source/ImgMedia/Public/ImgMediaSource.h
--- a/Engine/Plugins/Media/ImgMedia/Source/ImgMedia/Public/ImgMediaSource.h
+++ b/Engine/Plugins/Media/ImgMedia/Source/ImgMedia/Public/ImgMediaSource.h
@@ -6,6 +6,7 @@
 #include "Engine/EngineTypes.h"
 #include "Containers/UnrealString.h"
 #include "Misc/FrameRate.h"
+#include "Misc/Timecode.h"
 #include "UObject/ObjectMacros.h"
 #include "UObject/ScriptMacros.h"
 #include "BaseMediaSource.h"
@@ -74,6 +75,10 @@ public:
 
 #endif // WITH_EDITORONLY_DATA
 
+	/** Specification of a timecode associated with the start of the sequence. */
+	UPROPERTY(EditAnywhere, BlueprintReadOnly, Category="Timecode")
+	FTimecode StartTimecode;
+
 public:
 
 	/**
