Resolving static initialization order bug - use FAutoConsoleVariableRef instead of TAutoConsoleVariable. Use ECVF_RenderThreadSafe and ECVF_ReadOnly flags to match the intent of the previous functionality.

diff --git a/Engine/Plugins/Media/ImgMedia/Source/ImgMedia/Private/Readers/ExrImgMediaReaderGpu.cpp b/Engine/Plugins/Media/ImgMedia/Source/ImgMedia/Private/Readers/ExrImgMediaReaderGpu.cpp
--- a/Engine/Plugins/Media/ImgMedia/Source/ImgMedia/Private/Readers/ExrImgMediaReaderGpu.cpp
+++ b/Engine/Plugins/Media/ImgMedia/Source/ImgMedia/Private/Readers/ExrImgMediaReaderGpu.cpp
@@ -33,15 +33,13 @@ DECLARE_GPU_STAT_NAMED(ExrImgMediaReaderGpu_MipUpscale, TEXT("ExrImgGpu.MipRende
 DECLARE_GPU_STAT_NAMED(ExrImgMediaReaderGpu_CopyUploadBuffer, TEXT("ExrImgGpu.MipRender.UploadBufferCopy"));
 DECLARE_GPU_STAT_NAMED(ExrImgMediaReaderGpu_AllocateBuffer, TEXT("ExrImgGpu.MipRender.AllocateBuffer"));
 
-
-static TAutoConsoleVariable<bool> CVarExrReaderUseUploadHeap(
+static bool bExrReaderUseUploadHeap = true;
+static FAutoConsoleVariableRef CVarExrReaderUseUploadHeap(
 	TEXT("r.ExrReaderGPU.UseUploadHeap"),
-	true,
-	TEXT("Utilizes upload heap and copies raw exr buffer asynchronously.\n\
-			Requires a restart of the engine."));
-
-/** This is to force user to restart after they've changed this setting. */
-static bool bUseUploadHeap = CVarExrReaderUseUploadHeap.GetValueOnAnyThread();
+	bExrReaderUseUploadHeap,
+	TEXT("Utilizes upload heap and copies raw exr buffer asynchronously.\n")
+	TEXT("Read-only and to be set in a config file (requires restart)."),
+	ECVF_RenderThreadSafe | ECVF_ReadOnly);
 
 namespace {
 
@@ -177,7 +175,7 @@ FExrImgMediaReader::EReadResult FExrImgMediaReaderGpu::ReadMip
 			OutFrame->NumTilesRead++;
 		}
 
-		if (ReadResult == Success && bUseUploadHeap)
+		if (ReadResult == Success && bExrReaderUseUploadHeap)
 		{
 			ENQUEUE_RENDER_COMMAND(CopyFromUploadBuffer)([SampleConverter, BufferData, BufferRegionsToCopy, FrameId = ConverterParams.FrameId](FRHICommandListImmediate& RHICmdList)
 				{
@@ -584,7 +582,7 @@ void FExrImgMediaReaderGpu::CreateSampleConverterCallback(TSharedPtr<FExrMediaTe
 
 			if (BufferData.IsValid())
 			{
-				if (!BufferData->UploadBufferRef.IsValid() || (bUseUploadHeap && !BufferData->ShaderAccessBufferRef->IsValid()))
+				if (!BufferData->UploadBufferRef.IsValid() || (bExrReaderUseUploadHeap && !BufferData->ShaderAccessBufferRef->IsValid()))
 				{
 					continue;
 				}
@@ -659,7 +657,7 @@ FStructuredBufferPoolItemSharedPtr FExrImgMediaReaderGpu::AllocateGpuBufferFromP
 				AllocatedBuffer->UploadBufferRef = RHICmdList.CreateStructuredBuffer(sizeof(uint16) * 2., AllocSize, BUF_ShaderResource | BUF_Dynamic | BUF_FastVRAM, CreateInfo);
 				AllocatedBuffer->UploadBufferMapped = RHICmdList.LockBuffer(AllocatedBuffer->UploadBufferRef, 0, AllocSize, RLM_WriteOnly);
 
-				if (CVarExrReaderUseUploadHeap.GetValueOnAnyThread())
+				if (bExrReaderUseUploadHeap)
 				{
 					CreateInfo.DebugName = TEXT("ExrReaderGpu.DestBuffer");
 					AllocatedBuffer->ShaderAccessBufferRef = RHICmdList.CreateStructuredBuffer(sizeof(uint16) * 2., AllocSize, BUF_ShaderResource | BUF_FastVRAM, CreateInfo);
