ImgMedia: Initial implementation for DDS support with the image sequence media player. No mip map support currently.

Test 1080p frame read times (approximate) with varying formats, in increasing speed order:
* PNG (RGB): ~80ms
* EXR (16f ZIP RGBA, CPU reader): ~19ms
* JPG (RGB): ~10ms
* EXR (16f Uncompressed RGBA, GPU reader): ~4ms
* DDS (RGBA BC3/DXT5): ~0.5ms

The compressed DDS image sequence is both faster to read from disk & smaller than uncompressed 16-bit EXRs by a factor of roughly 8x. However, other bottlenecks such as copy/GPU costs are likely to show up at higher resolutions.

diff --git a/Engine/Plugins/Media/ImgMedia/Source/ImgMedia/ImgMedia.Build.cs b/Engine/Plugins/Media/ImgMedia/Source/ImgMedia/ImgMedia.Build.cs
--- a/Engine/Plugins/Media/ImgMedia/Source/ImgMedia/ImgMedia.Build.cs
+++ b/Engine/Plugins/Media/ImgMedia/Source/ImgMedia/ImgMedia.Build.cs
@@ -17,6 +17,7 @@ namespace UnrealBuildTool.Rules
 					"Core",
 					"CoreUObject",
 					"Engine",
+					"ImageCore",
 					"ImgMediaEngine",
 					"ImgMediaFactory",
 					"MediaUtils",
diff --git a/Engine/Plugins/Media/ImgMedia/Source/ImgMedia/Private/ImgMediaModule.cpp b/Engine/Plugins/Media/ImgMedia/Source/ImgMedia/Private/ImgMediaModule.cpp
--- a/Engine/Plugins/Media/ImgMedia/Source/ImgMedia/Private/ImgMediaModule.cpp
+++ b/Engine/Plugins/Media/ImgMedia/Source/ImgMedia/Private/ImgMediaModule.cpp
@@ -225,6 +225,7 @@ private:
 		TEXT("jpg"),
 		TEXT("jpeg"),
 		TEXT("png"),
+		TEXT("dds"),
 	};
 };
 
diff --git a/Engine/Plugins/Media/ImgMedia/Source/ImgMedia/Private/Loader/ImgMediaLoader.cpp b/Engine/Plugins/Media/ImgMedia/Source/ImgMedia/Private/Loader/ImgMediaLoader.cpp
--- a/Engine/Plugins/Media/ImgMedia/Source/ImgMedia/Private/Loader/ImgMediaLoader.cpp
+++ b/Engine/Plugins/Media/ImgMedia/Source/ImgMedia/Private/Loader/ImgMediaLoader.cpp
@@ -20,6 +20,7 @@
 #include "IImageWrapperModule.h"
 #include "IImgMediaModule.h"
 #include "Readers/IImgMediaReader.h"
+#include "Readers/DdsImgMediaReader.h"
 #include "ImgMediaLoaderWork.h"
 #include "Scheduler/ImgMediaScheduler.h"
 #include "Player/ImgMediaTextureSample.h"
@@ -840,6 +841,10 @@ bool FImgMediaLoader::LoadSequence(const FString& SequencePath, const FFrameRate
 		return false;
 #endif
 	}
+	else if (FirstExtension == TEXT("dds"))
+	{
+		Reader = MakeShared<FDdsImgMediaReader>(AsShared());
+	}
 	else
 	{
 		Reader = MakeShareable(new FGenericImgMediaReader(ImageWrapperModule, AsShared()));
@@ -1003,7 +1008,7 @@ void FImgMediaLoader::FindFiles(const FString& SequencePath, TArray<FString>& Ou
 	UE_LOG(LogImgMedia, Verbose, TEXT("Loader %p: Found %i image files in %s"), this, FoundFiles.Num(), *SequencePath);
 
 	// Same list of extensions as FGenericImgMediaReader::LoadFrameImage() with the addition of exr
-	const TSet<FString> SupportedImageFileExtensions = { TEXT("exr"), TEXT("jpg"), TEXT("jpeg"), TEXT("png"), TEXT("bmp") };
+	const TSet<FString> SupportedImageFileExtensions = { TEXT("exr"), TEXT("jpg"), TEXT("jpeg"), TEXT("png"), TEXT("bmp"), TEXT("dds") };
 
 	TArray<FString> UnnumberedFiles;
 	TSortedMap<int32, FString> SortedFoundFiles;
diff --git a/Engine/Plugins/Media/ImgMedia/Source/ImgMedia/Private/Readers/DdsImgMediaReader.cpp b/Engine/Plugins/Media/ImgMedia/Source/ImgMedia/Private/Readers/DdsImgMediaReader.cpp
new file mode 100644
--- /dev/null
+++ b/Engine/Plugins/Media/ImgMedia/Source/ImgMedia/Private/Readers/DdsImgMediaReader.cpp
@@ -0,0 +1,348 @@
+// Copyright Epic Games, Inc. All Rights Reserved.
+
+#include "DdsImgMediaReader.h"
+#include "ImgMediaPrivate.h"
+
+#include "IImageWrapper.h"
+#include "IImageWrapperModule.h"
+#include "ImageCoreUtils.h"
+#include "Loader/ImgMediaLoader.h"
+#include "Logging/StructuredLog.h"
+#include "Misc/FileHelper.h"
+#include "Misc/Paths.h"
+
+#include "DDSFile.h"
+
+namespace ImgMedia
+{
+	namespace DDS
+	{
+		/* Returns the media texture sample format from a DXGI format type. */
+		bool GetMediaTextureSampleFormat(const UE::DDS::EDXGIFormat& InDXGIFormat, EMediaTextureSampleFormat& OutFormat)
+		{
+			switch (InDXGIFormat)
+			{
+			case UE::DDS::EDXGIFormat::BC1_TYPELESS:
+			case UE::DDS::EDXGIFormat::BC1_UNORM:
+			case UE::DDS::EDXGIFormat::BC1_UNORM_SRGB:
+				OutFormat = EMediaTextureSampleFormat::DXT1;
+				break;
+			case UE::DDS::EDXGIFormat::BC3_TYPELESS:
+			case UE::DDS::EDXGIFormat::BC3_UNORM:
+			case UE::DDS::EDXGIFormat::BC3_UNORM_SRGB:
+				OutFormat = EMediaTextureSampleFormat::DXT5;
+				break;
+			case UE::DDS::EDXGIFormat::BC4_TYPELESS:
+			case UE::DDS::EDXGIFormat::BC4_UNORM:
+			case UE::DDS::EDXGIFormat::BC4_SNORM:
+				OutFormat = EMediaTextureSampleFormat::BC4;
+				break;
+
+			case UE::DDS::EDXGIFormat::R16G16B16A16_FLOAT:
+				OutFormat = EMediaTextureSampleFormat::FloatRGBA;
+				break;
+
+			default:
+				UE_LOG(LogImgMedia, Error, TEXT("Unsupported compression format, only BC1/DXT1, BC3/DXT5, BC4 & FloatRGBA are currently supported."));
+				return false;
+			}
+
+			return true;
+		}
+
+		/* Only consider mips larger than the block size as valid. */
+		inline bool IsMipValid(const UE::DDS::FDDSMip& InMip)
+		{
+			// See FD3D12Texture::UpdateTexture2D block size checks, and also
+			// TODO: D3D12_FEATURE_DATA_D3D12_OPTIONS8.UnalignedBlockTexturesSupported
+
+			return FMath::Min(InMip.Width, InMip.Height) >= 4 /* CompressionBlockSize */;
+		}
+
+		/* Returns the DDS total mip data size. */
+		SIZE_T GetTotalMipDataSize(const TUniquePtr<UE::DDS::FDDSFile>& DDS)
+		{
+			SIZE_T TotalMipDataSize = 0;
+			for (const UE::DDS::FDDSMip& Mip : DDS->Mips)
+			{
+				if (!IsMipValid(Mip))
+				{
+					break;
+				}
+
+				TotalMipDataSize += static_cast<SIZE_T>(Mip.DataSize);
+			}
+
+			return TotalMipDataSize;
+		}
+
+		/* Returns the DDS file mip count, accounting for engine maximum. */
+		int32 GetMipCount(const TUniquePtr<UE::DDS::FDDSFile>& DDS)
+		{
+			int32 MipCount = 0;
+
+			// Mips are ordered starting from mip 0 (full-size texture) decreasing in size
+			for (const UE::DDS::FDDSMip& Mip : DDS->Mips)
+			{
+				if (!IsMipValid(Mip))
+				{
+					break;
+				}
+
+				++MipCount;
+			}
+
+			return FMath::Min(MipCount, static_cast<int32>(MAX_TEXTURE_MIP_COUNT));
+		}
+
+		/* Returns whether the DDS file is sRGB-encoded. */
+		bool DDSPayloadIsSRGB(const TUniquePtr<UE::DDS::FDDSFile>& DDS, ERawImageFormat::Type RawFormat)
+		{
+			if (ERawImageFormat::GetFormatNeedsGammaSpace(RawFormat))
+			{
+				if (UE::DDS::DXGIFormatHasLinearAndSRGBForm(DDS->DXGIFormat))
+				{
+					return UE::DDS::DXGIFormatIsSRGB(DDS->DXGIFormat);
+				}
+
+				return true;
+			}
+
+			return false;
+		}
+
+		/* Returns DDS file frame information for image sequence playback. */
+		FImgMediaFrameInfo GetFrameInfo(const TUniquePtr<UE::DDS::FDDSFile>& DDS)
+		{
+			FImgMediaFrameInfo OutInfo;
+
+			ERawImageFormat::Type RawFormat = UE::DDS::DXGIFormatGetClosestRawFormat(DDS->DXGIFormat);
+
+			OutInfo.CompressionName = UE::DDS::DXGIFormatGetName(DDS->DXGIFormat);
+			OutInfo.Dim.X = DDS->Width;
+			OutInfo.Dim.Y = DDS->Height;
+			OutInfo.UncompressedSize = GetTotalMipDataSize(DDS);
+			OutInfo.NumMipLevels = GetMipCount(DDS);
+			OutInfo.FormatName = TEXT("DDS");
+			OutInfo.FrameRate = GetDefault<UImgMediaSettings>()->DefaultFrameRate;
+			OutInfo.Srgb = DDSPayloadIsSRGB(DDS, RawFormat);
+			OutInfo.NumChannels = 4;
+			OutInfo.bHasTiles = false;
+			OutInfo.TileDimensions = OutInfo.Dim;
+			OutInfo.NumTiles = FIntPoint(1, 1);
+			OutInfo.TileBorder = 0;
+
+			return OutInfo;
+		}
+	}
+}
+
+/* Convenience class for reading header/mip information and raw data from DDS files. */
+class FDdsReader
+{
+public:
+	/* Constructor */
+	FDdsReader(const FString& InFilename)
+		: Filename(InFilename)
+		, ArFileReader(IFileManager::Get().CreateFileReader(*InFilename))
+		, StartSeekPos(0)
+	{
+	}
+
+	/* Read file header and prepare for raw data reads. Must be called in advance of ReadRawMipData. */
+	TUniquePtr<UE::DDS::FDDSFile> ReadHeaderAndPrepare()
+	{
+		TArray64<uint8> HeaderData;
+		if (!ReadHeaderData(HeaderData))
+		{
+			return nullptr;
+		}
+
+		UE::DDS::EDDSError Error;
+		TUniquePtr<UE::DDS::FDDSFile> DDS = TUniquePtr<UE::DDS::FDDSFile>(UE::DDS::FDDSFile::CreateFromDDSInMemory(HeaderData.GetData(), HeaderData.Num(), &Error, UE::DDS::EDDSReadMipMode::HeaderWithMipInfo));
+		if (!DDS.IsValid())
+		{
+			check(Error != UE::DDS::EDDSError::OK);
+			if (Error != UE::DDS::EDDSError::NotADds && Error != UE::DDS::EDDSError::IoError)
+			{
+				UE_LOG(LogImgMedia, Warning, TEXT("Failed to load DDS (Error=%d) [%s]"), (int)Error, *Filename);
+			}
+			
+			return nullptr;
+		}
+
+		const int64 SizeOfFile = ArFileReader->TotalSize();
+		int64 TotalMipDataSize = 0;
+		for (const UE::DDS::FDDSMip& Mip : DDS->Mips)
+		{
+			TotalMipDataSize += Mip.DataSize;
+		}
+		StartSeekPos = SizeOfFile - TotalMipDataSize;
+		
+		return MoveTemp(DDS);
+	}
+
+	/* Read a specified size of raw file data. */
+	void ReadRawMipData(int64 MipDataOffset, void* MipData, int64 MipDataSize) const
+	{
+		if (ArFileReader.IsValid())
+		{
+			ArFileReader->Seek(StartSeekPos + MipDataOffset);
+			ArFileReader->Serialize(MipData, MipDataSize);
+		}
+	}
+
+	/* Close the file reader. */
+	void Close()
+	{
+		if (ArFileReader.IsValid())
+		{
+			ArFileReader->Close();
+		}
+	}
+
+private:
+
+	/* Convenience function to read the header data. */
+	bool ReadHeaderData(TArray64<uint8>& OutHeader) const
+	{
+		if (!ArFileReader.IsValid())
+		{
+			return false;
+		}
+
+		const int64 SizeOfFile = ArFileReader->TotalSize();
+		const int64 MinimalHeaderSize = UE::DDS::GetDDSHeaderMinimalSize();
+
+		// If the file is not bigger than the smallest header possible then clearly the file is not valid as a dds file.
+		if (SizeOfFile > MinimalHeaderSize)
+		{
+			const int64 MaximalHeaderSize = UE::DDS::GetDDSHeaderMaximalSize();
+			const int64 BytesToRead = FMath::Min(SizeOfFile, MaximalHeaderSize);
+
+			OutHeader.Reset(BytesToRead);
+			OutHeader.AddUninitialized(BytesToRead);
+			ArFileReader->Serialize(OutHeader.GetData(), OutHeader.Num());
+
+			return true;
+		}
+
+		return false;
+	}
+
+private:
+
+	// Loaded file name
+	FString Filename;
+	
+	// File reader
+	TUniquePtr<FArchive> ArFileReader;
+
+	// Start seek position for mip data
+	int64 StartSeekPos;
+};
+
+
+/* FGenericImgMediaReader structors
+ *****************************************************************************/
+
+FDdsImgMediaReader::FDdsImgMediaReader(const TSharedRef<FImgMediaLoader, ESPMode::ThreadSafe>& InLoader)
+	: LoaderPtr(InLoader)
+{ }
+
+FDdsImgMediaReader::~FDdsImgMediaReader() = default;
+
+/* FDdsImgMediaReader interface
+ *****************************************************************************/
+
+bool FDdsImgMediaReader::GetFrameInfo(const FString& ImagePath, FImgMediaFrameInfo& OutInfo)
+{
+	FDdsReader Reader(ImagePath);
+	TUniquePtr<UE::DDS::FDDSFile> DDS = Reader.ReadHeaderAndPrepare();
+	if (!DDS.IsValid())
+	{
+		UE_LOGFMT(LogImgMedia, Warning, "FDdsImgMediaReader: Failed to load image {0}", ImagePath);
+		return false;
+	}
+
+	const int32 MipCount = ImgMedia::DDS::GetMipCount(DDS);
+	if (MipCount > 1)
+	{
+		// TODO: Should be removed eventually, once reading & media texture issues are addressed.
+		UE_LOGFMT(LogImgMedia, Error, "FDdsImgMediaReader: Mips are not currently supported");
+		return false;
+	}
+
+	OutInfo = ImgMedia::DDS::GetFrameInfo(DDS);
+
+	return true;
+}
+
+
+bool FDdsImgMediaReader::ReadFrame(int32 FrameId, const TMap<int32, FImgMediaTileSelection>& InMipTiles, TSharedPtr<FImgMediaFrame, ESPMode::ThreadSafe> OutFrame)
+{
+	TRACE_CPUPROFILER_EVENT_SCOPE_TEXT(*FString::Printf(TEXT("DdsImgMedia.ReadFrame %d"), FrameId));
+
+	const TSharedPtr<FImgMediaLoader, ESPMode::ThreadSafe> Loader = LoaderPtr.Pin();
+	if (!Loader.IsValid())
+	{
+		return false;
+	}
+
+	if (InMipTiles.IsEmpty())
+	{
+		return false;
+	}
+
+	const FString& ImagePath = Loader->GetImagePath(FrameId, 0);
+	FDdsReader Reader(ImagePath);
+
+	TUniquePtr<UE::DDS::FDDSFile> DDS = Reader.ReadHeaderAndPrepare();
+	if (!DDS.IsValid())
+	{
+		UE_LOG(LogImgMedia, Warning, TEXT("FDdsImgMediaReader: Failed to load image %s"), *ImagePath);
+		return false;
+	}
+
+	int64 MipDataOffset = 0;
+	const int32 MipCount = ImgMedia::DDS::GetMipCount(DDS);
+
+	// Loop over all mips.
+	for (int32 MipLevel = 0; MipLevel < MipCount; ++MipLevel)
+	{
+		const int64 MipDataSize = DDS->Mips[MipLevel].DataSize;
+		const bool bMipLevelRequested = InMipTiles.Contains(MipLevel);
+		const bool bMipLevelCached = OutFrame->MipTilesPresent.Contains(MipLevel);
+
+		if (bMipLevelRequested && !bMipLevelCached)
+		{
+			const bool bBufferAllocated = OutFrame->Data.IsValid();
+
+			if (!bBufferAllocated)
+			{
+				const SIZE_T TotalMipDataSize = ImgMedia::DDS::GetTotalMipDataSize(DDS);
+				void* Buffer = FMemory::Malloc(TotalMipDataSize, PLATFORM_CACHE_LINE_SIZE);
+
+				EMediaTextureSampleFormat SampleFormat;
+				ImgMedia::DDS::GetMediaTextureSampleFormat(DDS->DXGIFormat, SampleFormat);
+
+				OutFrame->Data = MakeShareable(Buffer, [](void* ObjectToDelete) { FMemory::Free(ObjectToDelete); });
+				OutFrame->Stride = DDS->Mips[0].RowStride;
+				OutFrame->SetInfo(ImgMedia::DDS::GetFrameInfo(DDS));
+				OutFrame->Format = SampleFormat;
+
+				OutFrame->MipTilesPresent.Reset();
+			}
+
+			void* MipData = static_cast<void*>((uint8*)OutFrame->Data.Get() + static_cast<SIZE_T>(MipDataOffset));
+			Reader.ReadRawMipData(MipDataOffset, MipData, MipDataSize);
+
+			OutFrame->MipTilesPresent.Emplace(MipLevel, InMipTiles[MipLevel]);
+			OutFrame->NumTilesRead++;
+		}
+
+		MipDataOffset += MipDataSize;
+	}
+
+	return true;
+}
diff --git a/Engine/Plugins/Media/ImgMedia/Source/ImgMedia/Private/Readers/DdsImgMediaReader.h b/Engine/Plugins/Media/ImgMedia/Source/ImgMedia/Private/Readers/DdsImgMediaReader.h
new file mode 100644
--- /dev/null
+++ b/Engine/Plugins/Media/ImgMedia/Source/ImgMedia/Private/Readers/DdsImgMediaReader.h
@@ -0,0 +1,41 @@
+// Copyright Epic Games, Inc. All Rights Reserved.
+
+#pragma once
+
+#include "Containers/ContainersFwd.h"
+#include "IImgMediaReader.h"
+
+class FImgMediaLoader;
+
+/**
+ * Implements a reader for various image sequence formats.
+ */
+class FDdsImgMediaReader final
+	: public IImgMediaReader
+{
+public:
+
+	/**
+	 * Create and initialize a new instance.
+	 *
+	 * @param InImageWrapperModule The image wrapper module to use.
+	 */
+	FDdsImgMediaReader(const TSharedRef<FImgMediaLoader, ESPMode::ThreadSafe>& InLoader);
+
+	~FDdsImgMediaReader();
+
+public:
+
+	//~ IImgMediaReader interface
+
+	virtual bool GetFrameInfo(const FString& ImagePath, FImgMediaFrameInfo& OutInfo) override;
+	virtual bool ReadFrame(int32 FrameId, const TMap<int32, FImgMediaTileSelection>& InMipTiles, TSharedPtr<FImgMediaFrame, ESPMode::ThreadSafe> OutFrame) override;
+	virtual void CancelFrame(int32 FrameNumber) override {};
+	virtual void UncancelFrame(int32 FrameNumber) override {};
+
+private:
+
+	/** Our parent loader. */
+	TWeakPtr<FImgMediaLoader, ESPMode::ThreadSafe> LoaderPtr;
+};
+
