ImgMedia: Fix DDS sRGB flag support, removing image gamma space logic but replicating texture import conditions.

#rnx

diff --git a/Engine/Plugins/Media/ImgMedia/Source/ImgMedia/Private/Readers/DdsImgMediaReader.cpp b/Engine/Plugins/Media/ImgMedia/Source/ImgMedia/Private/Readers/DdsImgMediaReader.cpp
--- a/Engine/Plugins/Media/ImgMedia/Source/ImgMedia/Private/Readers/DdsImgMediaReader.cpp
+++ b/Engine/Plugins/Media/ImgMedia/Source/ImgMedia/Private/Readers/DdsImgMediaReader.cpp
@@ -96,19 +96,29 @@ namespace ImgMedia
 		}
 
 		/* Returns whether the DDS file is sRGB-encoded. */
-		bool DDSPayloadIsSRGB(const TUniquePtr<UE::DDS::FDDSFile>& DDS, ERawImageFormat::Type RawFormat)
+		bool DDSPayloadIsSRGB(const TUniquePtr<UE::DDS::FDDSFile>& DDS)
 		{
-			if (ERawImageFormat::GetFormatNeedsGammaSpace(RawFormat))
+			if (DDS->CreateFlags & UE::DDS::FDDSFile::CREATE_FLAG_WAS_D3D9)
 			{
-				if (UE::DDS::DXGIFormatHasLinearAndSRGBForm(DDS->DXGIFormat))
-				{
-					return UE::DDS::DXGIFormatIsSRGB(DDS->DXGIFormat);
-				}
-
+				// no SRGB info in Dx9 format
+				// assume SRGB yes
 				return true;
 			}
+			else if (UE::DDS::DXGIFormatHasLinearAndSRGBForm(DDS->DXGIFormat))
+			{
+				// Dx10 file with format that has linear/srgb pair
+				//	( _UNORM when _UNORM_SRGB)
+				return UE::DDS::DXGIFormatIsSRGB(DDS->DXGIFormat);
+			}
+			else
+			{
+				// Dx10 format that doesn't have linear/srgb pairs
 
-			return false;
+				// R8G8_UNORM and R8_UNORM have no _SRGB pair
+				// so no way to clearly indicate SRGB or Linear for them
+				// assume SRGB yes
+				return true;
+			}
 		}
 
 		/* Returns DDS file frame information for image sequence playback. */
@@ -116,8 +126,6 @@ namespace ImgMedia
 		{
 			FImgMediaFrameInfo OutInfo;
 
-			ERawImageFormat::Type RawFormat = UE::DDS::DXGIFormatGetClosestRawFormat(DDS->DXGIFormat);
-
 			OutInfo.CompressionName = UE::DDS::DXGIFormatGetName(DDS->DXGIFormat);
 			OutInfo.Dim.X = DDS->Width;
 			OutInfo.Dim.Y = DDS->Height;
@@ -125,7 +133,7 @@ namespace ImgMedia
 			OutInfo.NumMipLevels = GetMipCount(DDS);
 			OutInfo.FormatName = TEXT("DDS");
 			OutInfo.FrameRate = GetDefault<UImgMediaSettings>()->DefaultFrameRate;
-			OutInfo.Srgb = DDSPayloadIsSRGB(DDS, RawFormat);
+			OutInfo.Srgb = DDSPayloadIsSRGB(DDS);
 			OutInfo.NumChannels = 4;
 			OutInfo.bHasTiles = false;
 			OutInfo.TileDimensions = OutInfo.Dim;
