[Media] Img Sequence Scrubbing improvements

- Media Track: disable "BlockOnFetch" when scrubbing.
- ImgMedia: fix PeekVideoSampleTime to properly report when the sample is availabe so that Seek requests can be throttled when BlockOnFetch is disabled. Related to UE-197080 (cl 30423934).

diff --git a/Engine/Plugins/Media/ImgMedia/Source/ImgMedia/Private/Loader/ImgMediaLoader.cpp b/Engine/Plugins/Media/ImgMedia/Source/ImgMedia/Private/Loader/ImgMediaLoader.cpp
--- a/Engine/Plugins/Media/ImgMedia/Source/ImgMedia/Private/Loader/ImgMediaLoader.cpp
+++ b/Engine/Plugins/Media/ImgMedia/Source/ImgMedia/Private/Loader/ImgMediaLoader.cpp
@@ -602,7 +602,10 @@ IMediaSamples::EFetchBestSampleResult FImgMediaLoader::FetchBestVideoSampleForTi
 
 bool FImgMediaLoader::PeekVideoSampleTime(FMediaTimeStamp &TimeStamp, bool bIsLoopingEnabled, float PlayRate, const FTimespan& CurrentTime)
 {
-	// Strictly speaking we never actual return a "frame" here as we do NOT involve any requests to the cache -- we simply assume we could deliver what a "queue" would contain on demand later...
+	// When playing (i.e not seeking), we don't actually return a "frame" here as we do NOT involve any requests to the cache,
+	// we simply assume we could deliver what a "queue" would contain on demand later (in FetchBestVideoSampleForTimeRange).
+	// However, when seeking, we need to check the cache here so the seek status properly reflect the state of the cache and allow
+	// for seek throttling to work as intended.
 	if (IsInitialized())
 	{
 		check(QueuedSampleFetch.LastTimeStamp.IsValid());
@@ -627,6 +630,35 @@ bool FImgMediaLoader::PeekVideoSampleTime(FMediaTimeStamp &TimeStamp, bool bIsLo
 		}
 
 		TimeStamp.Time = FrameNumberToTime(TimeToFrameNumber(TimeStamp.Time));
+
+		// If waiting for a seek and there is no other playback blocking,
+		// properly indicate if the sample has arrived to allow for seek throttling.
+		if (QueuedSampleFetch.LastDuration.IsZero() && !bIsPlaybackBlocking)
+		{
+			const uint32 PeekFrameNumber = TimeToFrameNumber(TimeStamp.Time);
+
+			bool bHasFrame;
+			if (UseGlobalCache)
+			{
+				bHasFrame = GlobalCache->Contains(SequenceName, PeekFrameNumber);
+			}
+			else
+			{
+				FScopeLock ScopeLock(&CriticalSection);
+				bHasFrame = Frames.Contains(PeekFrameNumber);
+			}
+
+			if (!bHasFrame)
+			{
+				// Request the missing frame to avoid blocking.
+				RequestFrame(PeekFrameNumber, PlayRate, bIsLoopingEnabled);
+
+				// Need to invalidate the time stamp because FMediaPlayerFacade::UpdateSeekStatus ignores the return value.
+				TimeStamp.Invalidate();
+			}
+			return bHasFrame;
+		}
+
 		return true;
 	}
 	return false;
@@ -730,6 +762,8 @@ void FImgMediaLoader::Initialize(const FString& SequencePath, const FFrameRate&
 
 bool FImgMediaLoader::RequestFrame(FTimespan Time, float PlayRate, bool Loop)
 {
+	// Match the frame rounding done in FImgMediaLoader::Seek to ensure requested frame is the same as in PeekVideoSampleTime.
+	Time = FrameNumberToTime(TimeToFrameNumber(Time));
 	return RequestFrame(TimeToFrameNumber(Time), PlayRate, Loop);
 }
 
diff --git a/Engine/Plugins/Media/MediaCompositing/Source/MediaCompositing/Private/MovieScene/MovieSceneMediaTemplate.cpp b/Engine/Plugins/Media/MediaCompositing/Source/MediaCompositing/Private/MovieScene/MovieSceneMediaTemplate.cpp
--- a/Engine/Plugins/Media/MediaCompositing/Source/MediaCompositing/Private/MovieScene/MovieSceneMediaTemplate.cpp
+++ b/Engine/Plugins/Media/MediaCompositing/Source/MediaCompositing/Private/MovieScene/MovieSceneMediaTemplate.cpp
@@ -24,6 +24,11 @@
 
 #define MOVIESCENEMEDIATEMPLATE_TRACE_EVALUATION 0
 
+static TAutoConsoleVariable<bool> CVarBlockOnFetchScrub(
+	TEXT("MediaTrack.BlockOnFetchScrub"),
+	false,	// Should be disabled by default, but can be re-enabled in case of undesirable side effects in some use cases.
+	TEXT("Enable/Disable BlockOnFetch while scrubbing."),
+	ECVF_Default);
 
 /* Local helpers
  *****************************************************************************/
@@ -386,9 +391,17 @@ struct FMediaSectionExecutionToken
 			SeekPlayer(Context, MediaPlayer, MediaTime);
 		}
 
-		// Set blocking range / time-range to display
-		// (we always use the full current time for this, any adjustments to player timestamps are done internally)
-		MediaPlayer->SetBlockOnTimeRange(TRange<FTimespan>(CurrentTime, CurrentTime + FrameDuration));
+		if (Context.GetStatus() == EMovieScenePlayerStatus::Scrubbing && CVarBlockOnFetchScrub.GetValueOnGameThread() == false)
+		{
+			// When scrubbing, seek requests are non-blocking.
+			MediaPlayer->SetBlockOnTimeRange(TRange<FTimespan>::Empty());
+		}
+		else
+		{
+			// Set blocking range / time-range to display
+			// (we always use the full current time for this, any adjustments to player timestamps are done internally)
+			MediaPlayer->SetBlockOnTimeRange(TRange<FTimespan>(CurrentTime, CurrentTime + FrameDuration));
+		}
 	}
 
 private:
