Media: Add support for mips with block compression (DDS).

Also fixed support for 8-bit uncompressed formats.

diff --git a/Engine/Plugins/Media/ImgMedia/Source/ImgMedia/Private/Readers/DdsImgMediaReader.cpp b/Engine/Plugins/Media/ImgMedia/Source/ImgMedia/Private/Readers/DdsImgMediaReader.cpp
--- a/Engine/Plugins/Media/ImgMedia/Source/ImgMedia/Private/Readers/DdsImgMediaReader.cpp
+++ b/Engine/Plugins/Media/ImgMedia/Source/ImgMedia/Private/Readers/DdsImgMediaReader.cpp
@@ -18,7 +18,7 @@ namespace ImgMedia
 	namespace DDS
 	{
 		/* Returns the media texture sample format from a DXGI format type. */
-		bool GetMediaTextureSampleFormat(const UE::DDS::EDXGIFormat& InDXGIFormat, EMediaTextureSampleFormat& OutFormat)
+		bool GetMediaTextureSampleFormat(const UE::DDS::EDXGIFormat& InDXGIFormat, EMediaTextureSampleFormat& OutFormat, bool& bOutIsBlockCompressed)
 		{
 			switch (InDXGIFormat)
 			{
@@ -26,24 +26,42 @@ namespace ImgMedia
 			case UE::DDS::EDXGIFormat::BC1_UNORM:
 			case UE::DDS::EDXGIFormat::BC1_UNORM_SRGB:
 				OutFormat = EMediaTextureSampleFormat::DXT1;
+				bOutIsBlockCompressed = true;
 				break;
 			case UE::DDS::EDXGIFormat::BC3_TYPELESS:
 			case UE::DDS::EDXGIFormat::BC3_UNORM:
 			case UE::DDS::EDXGIFormat::BC3_UNORM_SRGB:
 				OutFormat = EMediaTextureSampleFormat::DXT5;
+				bOutIsBlockCompressed = true;
 				break;
 			case UE::DDS::EDXGIFormat::BC4_TYPELESS:
 			case UE::DDS::EDXGIFormat::BC4_UNORM:
 			case UE::DDS::EDXGIFormat::BC4_SNORM:
 				OutFormat = EMediaTextureSampleFormat::BC4;
+				bOutIsBlockCompressed = true;
+				break;
+			case UE::DDS::EDXGIFormat::B8G8R8A8_TYPELESS:
+			case UE::DDS::EDXGIFormat::B8G8R8A8_UNORM:
+			case UE::DDS::EDXGIFormat::R8G8B8A8_SNORM:
+			case UE::DDS::EDXGIFormat::B8G8R8A8_UNORM_SRGB:
+				OutFormat = EMediaTextureSampleFormat::CharBGRA;
+				bOutIsBlockCompressed = false;
+				break;
+			case UE::DDS::EDXGIFormat::R8G8B8A8_TYPELESS:
+			case UE::DDS::EDXGIFormat::R8G8B8A8_UNORM:
+			case UE::DDS::EDXGIFormat::R8G8B8A8_UNORM_SRGB:
+				OutFormat = EMediaTextureSampleFormat::CharRGBA;
+				bOutIsBlockCompressed = false;
 				break;
-
 			case UE::DDS::EDXGIFormat::R16G16B16A16_FLOAT:
 				OutFormat = EMediaTextureSampleFormat::FloatRGBA;
+				bOutIsBlockCompressed = false;
 				break;
 
 			default:
 				UE_LOG(LogImgMedia, Error, TEXT("Unsupported compression format, only BC1/DXT1, BC3/DXT5, BC4 & FloatRGBA are currently supported."));
+				OutFormat = EMediaTextureSampleFormat::Undefined;
+				bOutIsBlockCompressed = false;
 				return false;
 			}
 
@@ -273,12 +291,22 @@ bool FDdsImgMediaReader::GetFrameInfo(const FString& ImagePath, FImgMediaFrameIn
 		return false;
 	}
 
+	EMediaTextureSampleFormat SampleFormat;
+	bool bIsBlockCompressed;
+	if (!ImgMedia::DDS::GetMediaTextureSampleFormat(DDS->DXGIFormat, SampleFormat, bIsBlockCompressed))
+	{
+		UE_LOGFMT(LogImgMedia, Error, "FDdsImgMediaReader: Texture format {0} is not currently supported.", UE::DDS::DXGIFormatGetName(DDS->DXGIFormat));
+		return false;
+	}
+
 	const int32 MipCount = ImgMedia::DDS::GetMipCount(DDS);
 	if (MipCount > 1)
 	{
-		// TODO: Should be removed eventually, once reading & media texture issues are addressed.
-		UE_LOGFMT(LogImgMedia, Error, "FDdsImgMediaReader: Mips are not currently supported");
-		return false;
+		if (bIsBlockCompressed && (!FMath::IsPowerOfTwo(DDS->Width) || !FMath::IsPowerOfTwo(DDS->Height)))
+		{
+			UE_LOGFMT(LogImgMedia, Error, "FDdsImgMediaReader: Compressed textures with mip maps currently need to have power-of-two dimensions.");
+			return false;
+		}
 	}
 
 	OutInfo = ImgMedia::DDS::GetFrameInfo(DDS);
@@ -332,7 +360,8 @@ bool FDdsImgMediaReader::ReadFrame(int32 FrameId, const TMap<int32, FImgMediaTil
 				void* Buffer = FMemory::Malloc(TotalMipDataSize, PLATFORM_CACHE_LINE_SIZE);
 
 				EMediaTextureSampleFormat SampleFormat;
-				ImgMedia::DDS::GetMediaTextureSampleFormat(DDS->DXGIFormat, SampleFormat);
+				bool bIsBlockCompressed;
+				ImgMedia::DDS::GetMediaTextureSampleFormat(DDS->DXGIFormat, SampleFormat, bIsBlockCompressed);
 
 				OutFrame->Data = MakeShareable(Buffer, [](void* ObjectToDelete) { FMemory::Free(ObjectToDelete); });
 				OutFrame->Stride = DDS->Mips[0].RowStride;
@@ -352,5 +381,17 @@ bool FDdsImgMediaReader::ReadFrame(int32 FrameId, const TMap<int32, FImgMediaTil
 		MipDataOffset += MipDataSize;
 	}
 
+	/*
+	* TODO: We currently rely on the media texture copy sample logic to update the
+	* texture resource mips, which needlessly copies mips that have not been read.
+	*
+	* Using our own sample converter (like FExrMediaTextureSampleConverter would solve
+	* this issue if only the media texture resource did not force conversion to float on
+	* DXT formats. See GetConvertedPixelFormat & FMediaTextureResource::RequiresConversion.
+	* 
+	* Changing these functions without further refactoring to allow PF_DXT1 & PF_DXT5 would
+	* currently break other players however.
+	*/
+
 	return true;
 }
diff --git a/Engine/Source/Runtime/MediaAssets/Private/Misc/MediaTextureResource.cpp b/Engine/Source/Runtime/MediaAssets/Private/Misc/MediaTextureResource.cpp
--- a/Engine/Source/Runtime/MediaAssets/Private/Misc/MediaTextureResource.cpp
+++ b/Engine/Source/Runtime/MediaAssets/Private/Misc/MediaTextureResource.cpp
@@ -1454,13 +1454,19 @@ void FMediaTextureResource::CopySample(FRHICommandListImmediate& RHICmdList, con
 			// Copy sample data (from CPU mem) to output render target
 			const FIntPoint SampleDim = Sample->GetDim();
 			const uint8* Data = (const uint8*)Sample->GetBuffer();
+			const FPixelFormatInfo& FormatInfo = GPixelFormats[MediaTextureResourceHelpers::GetPixelFormat(Sample)];
+
 			for (uint8 MipLevel = 0; MipLevel < SampleNumMips; ++MipLevel)
 			{
-				uint32 Stride = Sample->GetStride() >> MipLevel;
+				uint32 Width = SampleDim.X >> MipLevel;
 				uint32 Height = SampleDim.Y >> MipLevel;
-				FUpdateTextureRegion2D Region(0, 0, 0, 0, SampleDim.X >> MipLevel, Height);
+				uint32 Stride = Sample->GetStride() >> MipLevel;
+
+				FUpdateTextureRegion2D Region(0, 0, 0, 0, Width, Height);
 				RHIUpdateTexture2D(RenderTargetTextureRHI, MipLevel, Region, Stride, Data);
-				Data += Stride * Height;
+				
+				uint32 HeightUnit = (Height + FormatInfo.BlockSizeY - 1) / FormatInfo.BlockSizeY;
+				Data += Stride * HeightUnit;
 			}
 
 			// Make sure resource is in SRV mode again
