ImgMedia: Add support for source color space, encoding & chromatic adaptation on img media media sources.

We also update the default color space for image sequences to be the working color space, like texture assets.

diff --git a/Engine/Plugins/Media/ImgMedia/Source/ImgMedia/Internal/Loader/ImgMediaLoader.h b/Engine/Plugins/Media/ImgMedia/Source/ImgMedia/Internal/Loader/ImgMediaLoader.h
--- a/Engine/Plugins/Media/ImgMedia/Source/ImgMedia/Internal/Loader/ImgMediaLoader.h
+++ b/Engine/Plugins/Media/ImgMedia/Source/ImgMedia/Internal/Loader/ImgMediaLoader.h
@@ -23,6 +23,7 @@ class IQueuedWork;
 
 struct FImgMediaFrame;
 struct FImgMediaFrameInfo;
+struct FNativeMediaSourceColorSettings;
 
 /**
  * Settings for the smart cache.
@@ -86,6 +87,7 @@ public:
 	FImgMediaLoader(const TSharedRef<FImgMediaScheduler, ESPMode::ThreadSafe>& InScheduler,
 		const TSharedRef<FImgMediaGlobalCache, ESPMode::ThreadSafe>& InGlobalCache,
 		const TSharedPtr<FImgMediaMipMapInfo, ESPMode::ThreadSafe>& InMipMapInfo,
+		const TSharedPtr<FNativeMediaSourceColorSettings, ESPMode::ThreadSafe>& InSourceColorSettings,
 		bool bInFillGapsInSequence,
 		const FImgMediaLoaderSmartCacheSettings& InSmartCacheSettings);
 
@@ -631,6 +633,9 @@ private:
 	/** MipMapInfo object used to handle mipmaps. Could be null if we have no mipmaps. */
 	TSharedPtr<FImgMediaMipMapInfo, ESPMode::ThreadSafe> MipMapInfo;
 
+	/** Container for media source color settings. */
+	TSharedPtr<FNativeMediaSourceColorSettings, ESPMode::ThreadSafe> SourceColorSettings;
+
 	/** Width and height of the image sequence (in pixels) .*/
 	FIntPoint SequenceDim;
 
diff --git a/Engine/Plugins/Media/ImgMedia/Source/ImgMedia/Private/Assets/ImgMediaSource.cpp b/Engine/Plugins/Media/ImgMedia/Source/ImgMedia/Private/Assets/ImgMediaSource.cpp
--- a/Engine/Plugins/Media/ImgMedia/Source/ImgMedia/Private/Assets/ImgMediaSource.cpp
+++ b/Engine/Plugins/Media/ImgMedia/Source/ImgMedia/Private/Assets/ImgMediaSource.cpp
@@ -14,6 +14,62 @@
 
 #include UE_INLINE_GENERATED_CPP_BY_NAME(ImgMediaSource)
 
+#if WITH_EDITOR
+void FMediaSourceColorSettings::UpdateColorSpaceChromaticities()
+{
+	if (ColorSpaceOverride != ETextureColorSpace::TCS_Custom)
+	{
+		UE::Color::FColorSpace Chromaticities(static_cast<UE::Color::EColorSpace>(ColorSpaceOverride));
+		Chromaticities.GetChromaticities(RedChromaticityCoordinate, GreenChromaticityCoordinate, BlueChromaticityCoordinate, WhiteChromaticityCoordinate);
+	}
+}
+#endif // WITH_EDITOR
+
+FNativeMediaSourceColorSettings::FNativeMediaSourceColorSettings()
+	: EncodingOverride(UE::Color::EEncoding::None)
+	, ColorSpaceOverride()
+	, ChromaticAdaptationMethod(UE::Color::EChromaticAdaptationMethod::Bradford)
+{}
+
+void FNativeMediaSourceColorSettings::Update(const FMediaSourceColorSettings& InSettings)
+{
+	{
+		FScopeLock Lock(&ColorSpaceCriticalSection);
+
+		if (InSettings.ColorSpaceOverride == ETextureColorSpace::TCS_Custom)
+		{
+			ColorSpaceOverride = UE::Color::FColorSpace(
+				InSettings.RedChromaticityCoordinate,
+				InSettings.GreenChromaticityCoordinate,
+				InSettings.BlueChromaticityCoordinate,
+				InSettings.WhiteChromaticityCoordinate
+			);
+		}
+		else if (InSettings.ColorSpaceOverride != ETextureColorSpace::TCS_None)
+		{
+			ColorSpaceOverride = UE::Color::FColorSpace(static_cast<UE::Color::EColorSpace>(InSettings.ColorSpaceOverride));
+		}
+		else
+		{
+			ColorSpaceOverride.Reset();
+		}
+	}
+
+	EncodingOverride = static_cast<UE::Color::EEncoding>(InSettings.EncodingOverride);
+	ChromaticAdaptationMethod = static_cast<UE::Color::EChromaticAdaptationMethod>(InSettings.ChromaticAdaptationMethod);
+}
+
+const UE::Color::FColorSpace& FNativeMediaSourceColorSettings::GetColorSpaceOverride(const UE::Color::FColorSpace& InDefaultColorSpace) const
+{
+	FScopeLock Lock(&ColorSpaceCriticalSection);
+	
+	if (ColorSpaceOverride.IsSet())
+	{
+		return ColorSpaceOverride.GetValue();
+	}
+
+	return InDefaultColorSpace;
+}
 
 /* UImgMediaSource structors
  *****************************************************************************/
@@ -23,6 +79,7 @@ UImgMediaSource::UImgMediaSource()
 	, FrameRateOverride(0, 0)
 	, bFillGapsInSequence(true)
 	, MipMapInfo(MakeShared<FImgMediaMipMapInfo, ESPMode::ThreadSafe>())
+	, NativeSourceColorSettings(MakeShared<FNativeMediaSourceColorSettings, ESPMode::ThreadSafe>())
 {
 }
 
@@ -175,6 +232,13 @@ TSharedPtr<IMediaOptions::FDataContainer, ESPMode::ThreadSafe> UImgMediaSource::
 	{
 		return MipMapInfo;
 	}
+	
+	if (Key == ImgMedia::SourceColorSettingsOption)
+	{
+		NativeSourceColorSettings->Update(SourceColorSettings);
+
+		return NativeSourceColorSettings;
+	}
 
 	return Super::GetMediaOption(Key, DefaultValue);
 }
@@ -185,7 +249,8 @@ bool UImgMediaSource::HasMediaOption(const FName& Key) const
 		(Key == ImgMedia::FrameRateOverrideDenonimatorOption) ||
 		(Key == ImgMedia::FrameRateOverrideNumeratorOption) ||
 		(Key == ImgMedia::ProxyOverrideOption) ||
-		(Key == ImgMedia::MipMapInfoOption))
+		(Key == ImgMedia::MipMapInfoOption) ||
+		(Key == ImgMedia::SourceColorSettingsOption))
 	{
 		return true;
 	}
@@ -226,6 +291,15 @@ void UImgMediaSource::PostEditChangeProperty(FPropertyChangedEvent& PropertyChan
 	{
 		GenerateThumbnail();
 	}
+	else if (PropertyChangedEvent.GetPropertyName() == GET_MEMBER_NAME_CHECKED(FMediaSourceColorSettings, ColorSpaceOverride))
+	{
+		SourceColorSettings.UpdateColorSpaceChromaticities();
+	}
+
+	if (PropertyChangedEvent.GetMemberPropertyName() == GET_MEMBER_NAME_CHECKED(ThisClass, SourceColorSettings))
+	{
+		NativeSourceColorSettings->Update(SourceColorSettings);
+	}
 
 	Super::PostEditChangeProperty(PropertyChangedEvent);
 }
diff --git a/Engine/Plugins/Media/ImgMedia/Source/ImgMedia/Private/ImgMediaPrivate.h b/Engine/Plugins/Media/ImgMedia/Source/ImgMedia/Private/ImgMediaPrivate.h
--- a/Engine/Plugins/Media/ImgMedia/Source/ImgMedia/Private/ImgMediaPrivate.h
+++ b/Engine/Plugins/Media/ImgMedia/Source/ImgMedia/Private/ImgMediaPrivate.h
@@ -54,4 +54,7 @@ namespace ImgMedia
 
 	/** Name of the NumTilesY media option. */
 	static const FName NumTilesYOption("ImgNumTilesY");
+
+	/** Name of the SourceColorSettings media option. */
+	static const FName SourceColorSettingsOption("SourceColorSettings");
 }
diff --git a/Engine/Plugins/Media/ImgMedia/Source/ImgMedia/Private/Loader/ImgMediaLoader.cpp b/Engine/Plugins/Media/ImgMedia/Source/ImgMedia/Private/Loader/ImgMediaLoader.cpp
--- a/Engine/Plugins/Media/ImgMedia/Source/ImgMedia/Private/Loader/ImgMediaLoader.cpp
+++ b/Engine/Plugins/Media/ImgMedia/Source/ImgMedia/Private/Loader/ImgMediaLoader.cpp
@@ -22,6 +22,7 @@
 #include "Readers/IImgMediaReader.h"
 #include "Readers/DdsImgMediaReader.h"
 #include "ImgMediaLoaderWork.h"
+#include "ImgMediaSource.h"
 #include "Scheduler/ImgMediaScheduler.h"
 #include "Player/ImgMediaTextureSample.h"
 
@@ -158,6 +159,7 @@ void FImgMediaLoaderBandwidth::EmptyCache()
 FImgMediaLoader::FImgMediaLoader(const TSharedRef<FImgMediaScheduler, ESPMode::ThreadSafe>& InScheduler,
 	const TSharedRef<FImgMediaGlobalCache, ESPMode::ThreadSafe>& InGlobalCache,
 	const TSharedPtr<FImgMediaMipMapInfo, ESPMode::ThreadSafe>& InMipMapInfo,
+	const TSharedPtr<FNativeMediaSourceColorSettings, ESPMode::ThreadSafe>& InSourceColorSettings,
 	bool bInFillGapsInSequence,
 	const FImgMediaLoaderSmartCacheSettings& InSmartCacheSettings)
 	: Frames(1)
@@ -170,6 +172,7 @@ FImgMediaLoader::FImgMediaLoader(const TSharedRef<FImgMediaScheduler, ESPMode::T
 	, Scheduler(InScheduler)
 	, GlobalCache(InGlobalCache)
 	, MipMapInfo(InMipMapInfo)
+	, SourceColorSettings(InSourceColorSettings)
 	, SequenceDim(FIntPoint::ZeroValue)
 	, SequenceDuration(FTimespan::Zero())
 	, SequenceFrameRate(0, 0)
@@ -282,11 +285,8 @@ TSharedPtr<FImgMediaTextureSample, ESPMode::ThreadSafe> FImgMediaLoader::GetFram
 	auto Sample = MakeShared<FImgMediaTextureSample, ESPMode::ThreadSafe>();
 
 	ImgMediaLoader::CheckAndUpdateImgDimensions(SequenceDim, Frame->Get()->GetDim());
-
-	if (!Sample->Initialize(*Frame->Get(), SequenceDim, FMediaTimeStamp(FrameStartTime, 0), NextStartTime - FrameStartTime, GetNumMipLevels(), TilingDescription))
-	{
-		return nullptr;
-	}
+	
+	Sample->Initialize(*Frame->Get(), SequenceDim, FMediaTimeStamp(FrameStartTime, 0), NextStartTime - FrameStartTime, GetNumMipLevels(), TilingDescription, SourceColorSettings);
 
 	return Sample;
 }
@@ -559,12 +559,11 @@ IMediaSamples::EFetchBestSampleResult FImgMediaLoader::FetchBestVideoSampleForTi
 
 			ImgMediaLoader::CheckAndUpdateImgDimensions(SequenceDim, Frame->Get()->GetDim());
 
-			if (Sample->Initialize(*Frame->Get(), SequenceDim, QueuedSampleFetch.LastTimeStamp, QueuedSampleFetch.LastDuration, GetNumMipLevels(), TilingDescription))
-			{
-				OutSample = Sample;
-				CSV_EVENT(ImgMedia, TEXT("LoaderFetchHit %d %d-%d"), MaxIdx, UnwrappedStartIndex % NumFrames, UnwrappedEndIndex % NumFrames);
-				return IMediaSamples::EFetchBestSampleResult::Ok;
-			}
+			Sample->Initialize(*Frame->Get(), SequenceDim, QueuedSampleFetch.LastTimeStamp, QueuedSampleFetch.LastDuration, GetNumMipLevels(), TilingDescription, SourceColorSettings);
+
+			OutSample = Sample;
+			CSV_EVENT(ImgMedia, TEXT("LoaderFetchHit %d %d-%d"), MaxIdx, UnwrappedStartIndex% NumFrames, UnwrappedEndIndex% NumFrames);
+			return IMediaSamples::EFetchBestSampleResult::Ok;
 		}
 		else
 		{
diff --git a/Engine/Plugins/Media/ImgMedia/Source/ImgMedia/Private/Player/ImgMediaPlayer.cpp b/Engine/Plugins/Media/ImgMedia/Source/ImgMedia/Private/Player/ImgMediaPlayer.cpp
--- a/Engine/Plugins/Media/ImgMedia/Source/ImgMedia/Private/Player/ImgMediaPlayer.cpp
+++ b/Engine/Plugins/Media/ImgMedia/Source/ImgMedia/Private/Player/ImgMediaPlayer.cpp
@@ -240,6 +240,7 @@ bool FImgMediaPlayer::Open(const FString& Url, const IMediaOptions* Options, con
 	// get frame rate override, if any
 	FFrameRate FrameRateOverride(0, 0);
 	TSharedPtr<FImgMediaMipMapInfo, ESPMode::ThreadSafe> MipMapInfo;
+	TSharedPtr<FNativeMediaSourceColorSettings, ESPMode::ThreadSafe> SourceColorSettings;
 	bool bFillGapsInSequence = true;
 	bool bIsSmartCacheEnabled = GetBoolOption(PlayerOptions, MediaPlayerOptionValues::ImgMediaSmartCacheEnabled(), false);
 	double SmartCacheTimeToLookAhead = GetDoubleOption(PlayerOptions, MediaPlayerOptionValues::ImgMediaSmartCacheTimeToLookAhead(), 0.0);
@@ -280,12 +281,18 @@ bool FImgMediaPlayer::Open(const FString& Url, const IMediaOptions* Options, con
 				}
 			}
 		}
+
+		DataContainer  = Options->GetMediaOption(ImgMedia::SourceColorSettingsOption, DefaultValue);
+		if (DataContainer.IsValid())
+		{
+			SourceColorSettings = StaticCastSharedPtr<FNativeMediaSourceColorSettings, IMediaOptions::FDataContainer, ESPMode::ThreadSafe>(DataContainer);
+		}
 	}
 
 	// initialize image loader on a separate thread
 	FImgMediaLoaderSmartCacheSettings SmartCacheSettings(bIsSmartCacheEnabled, (float)SmartCacheTimeToLookAhead);
 	Loader = MakeShared<FImgMediaLoader, ESPMode::ThreadSafe>(Scheduler.ToSharedRef(),
-		GlobalCache.ToSharedRef(), MipMapInfo, bFillGapsInSequence, SmartCacheSettings);
+		GlobalCache.ToSharedRef(), MipMapInfo, SourceColorSettings, bFillGapsInSequence, SmartCacheSettings);
 	Scheduler->RegisterLoader(Loader.ToSharedRef());
 
 	const FString SequencePath = Url.RightChop(6);
diff --git a/Engine/Plugins/Media/ImgMedia/Source/ImgMedia/Private/Player/ImgMediaTextureSample.h b/Engine/Plugins/Media/ImgMedia/Source/ImgMedia/Private/Player/ImgMediaTextureSample.h
--- a/Engine/Plugins/Media/ImgMedia/Source/ImgMedia/Private/Player/ImgMediaTextureSample.h
+++ b/Engine/Plugins/Media/ImgMedia/Source/ImgMedia/Private/Player/ImgMediaTextureSample.h
@@ -8,6 +8,7 @@
 #include "Misc/Timespan.h"
 #include "Templates/SharedPointer.h"
 
+#include "ImgMediaSource.h"
 #include "Readers/IImgMediaReader.h"
 
 
@@ -24,7 +25,9 @@ public:
 		: Duration(FTimespan::Zero())
 		, OutputDim(FIntPoint::ZeroValue)
 		, Time(FTimespan::Zero())
+		, NumMipMaps(1)
 		, TilingDesc()
+		, SourceColorSettings()
 	{ }
 
 public:
@@ -36,8 +39,18 @@ public:
 	 * @param InOutputDim The sample's output width and height (in pixels).
 	 * @param InTime The sample time (in the player's local clock).
 	 * @param InDuration The duration for which the sample is valid.
+	 * @param InNumMipMaps The sample's number of mip levels.
+	 * @param InTilingDesc The sample's tiling description.
+	 * @param InSourceColorSettings The sample's source color space & encoding overrides.
 	 */
-	bool Initialize(FImgMediaFrame& InFrame, const FIntPoint& InOutputDim, FMediaTimeStamp InTime, FTimespan InDuration, uint8 InNumMipMaps, FMediaTextureTilingDescription InTilingDesc = FMediaTextureTilingDescription())
+	void Initialize(
+		FImgMediaFrame& InFrame,
+		const FIntPoint& InOutputDim,
+		FMediaTimeStamp InTime,
+		FTimespan InDuration,
+		uint8 InNumMipMaps,
+		FMediaTextureTilingDescription InTilingDesc,
+		TSharedPtr<FNativeMediaSourceColorSettings, ESPMode::ThreadSafe> InSourceColorSettings)
 	{
 		Duration = InDuration;
 		Frame = InFrame;
@@ -45,6 +58,7 @@ public:
 		Time = InTime;
 		NumMipMaps = InNumMipMaps;
 		TilingDesc = InTilingDesc;
+		SourceColorSettings = InSourceColorSettings;
 
 		// If we have no data then make sure the number of mipmaps is 1.
 		// otherwise FMediaTextureResource won't like it.
@@ -52,8 +66,6 @@ public:
 		{
 			NumMipMaps = 1;
 		}
-
-		return true;
 	}
 
 public:
@@ -124,9 +136,57 @@ public:
 
 	virtual bool IsOutputSrgb() const override
 	{
+		if (SourceColorSettings.IsValid())
+		{
+			const UE::Color::EEncoding EncodingOverride = SourceColorSettings->GetEncodingOverride();
+
+			if (EncodingOverride != UE::Color::EEncoding::None)
+			{
+				return EncodingOverride == UE::Color::EEncoding::sRGB;
+			}
+		}
+
 		return Frame.IsOutputSrgb();
 	}
 
+	virtual const UE::Color::FColorSpace& GetSourceColorSpace() const override
+	{
+		// Like texture assets, we assume images are already in working color space.
+		const UE::Color::FColorSpace& DefaultColorSpace = UE::Color::FColorSpace::GetWorking();
+
+		if (SourceColorSettings.IsValid())
+		{
+			return SourceColorSettings->GetColorSpaceOverride(DefaultColorSpace);
+		}
+
+		return DefaultColorSpace;
+	}
+
+	virtual UE::Color::EEncoding GetEncodingType() const override
+	{
+		if (SourceColorSettings.IsValid())
+		{
+			const UE::Color::EEncoding EncodingOverride = SourceColorSettings->GetEncodingOverride();
+
+			if (EncodingOverride != UE::Color::EEncoding::None)
+			{
+				return EncodingOverride;
+			}
+		}
+
+		return IMediaTextureSample::GetEncodingType();
+	}
+
+	virtual UE::Color::EChromaticAdaptationMethod GetChromaticAdapationMethod() const
+	{
+		if (SourceColorSettings.IsValid())
+		{
+			return SourceColorSettings->GetChromaticAdaptationMethod();
+		}
+
+		return IMediaTextureSample::GetChromaticAdapationMethod();
+	}
+
 private:
 
 	/** Duration for which the sample is valid. */
@@ -146,4 +206,7 @@ private:
 
 	/** Description of the number and size of tiles in this sample. */
 	FMediaTextureTilingDescription TilingDesc;
+
+	/** Manual source color space & encoding overrides. */
+	TSharedPtr<FNativeMediaSourceColorSettings, ESPMode::ThreadSafe> SourceColorSettings;
 };
diff --git a/Engine/Plugins/Media/ImgMedia/Source/ImgMedia/Public/ImgMediaSource.h b/Engine/Plugins/Media/ImgMedia/Source/ImgMedia/Public/ImgMediaSource.h
--- a/Engine/Plugins/Media/ImgMedia/Source/ImgMedia/Public/ImgMediaSource.h
+++ b/Engine/Plugins/Media/ImgMedia/Source/ImgMedia/Public/ImgMediaSource.h
@@ -3,7 +3,9 @@
 #pragma once
 
 #include "CoreTypes.h"
+#include "ColorManagement/ColorSpace.h"
 #include "Engine/EngineTypes.h"
+#include "Engine/TextureDefines.h"
 #include "Containers/UnrealString.h"
 #include "Misc/FrameRate.h"
 #include "Misc/Timecode.h"
@@ -23,6 +25,103 @@ struct IMGMEDIA_API FImgMediaSourceCustomizationSequenceProxy
 	GENERATED_BODY()
 };
 
+// TODO: EMediaSourceEncoding, FMediaSourceColorSettings & FNativeMediaSourceColorSettings should be moved up to BaseMediaSource.h and unified with UCaptureCardMediaSource settings.
+
+/** List of source encodings that can be converted to linear. (Integer values match the ETextureSourceEncoding values in TextureDefines.h */
+UENUM()
+enum class EMediaSourceEncoding : uint8
+{
+	MSE_None	= 0 UMETA(DisplayName = "None", ToolTip = "The source encoding is not overridden."),
+	MSE_Linear	= 1 UMETA(DisplayName = "Linear", ToolTip = "The source encoding is considered linear."),
+	MSE_sRGB	= 2 UMETA(DisplayName = "sRGB", ToolTip = "sRGB source encoding to be linearized"),
+	MSE_ST2084	= 3 UMETA(DisplayName = "ST 2084/PQ", ToolTip = "SMPTE ST 2084/PQ source encoding to be linearized"),
+	MSE_SLog3	= 12 UMETA(DisplayName = "SLog3", ToolTip = "Sony SLog3 source encoding to be linearized"),
+
+	MSE_MAX,
+};
+
+/* Manual definition of media source color space & encoding. */
+USTRUCT(BlueprintType)
+struct FMediaSourceColorSettings
+{
+	GENERATED_USTRUCT_BODY()
+
+	FMediaSourceColorSettings()
+		: EncodingOverride(EMediaSourceEncoding::MSE_None)
+		, ColorSpaceOverride(ETextureColorSpace::TCS_None)
+		, RedChromaticityCoordinate(FVector2D::ZeroVector)
+		, GreenChromaticityCoordinate(FVector2D::ZeroVector)
+		, BlueChromaticityCoordinate(FVector2D::ZeroVector)
+		, WhiteChromaticityCoordinate(FVector2D::ZeroVector)
+		, ChromaticAdaptationMethod(ETextureChromaticAdaptationMethod::TCAM_Bradford)
+	{}
+
+	/** Source encoding of the media. */
+	UPROPERTY(EditAnywhere, BlueprintReadWrite, Category = ColorManagement)
+	EMediaSourceEncoding EncodingOverride;
+
+	/** Source color space of the media. */
+	UPROPERTY(EditAnywhere, BlueprintReadWrite, Category = ColorManagement)
+	ETextureColorSpace ColorSpaceOverride;
+
+	/** Red chromaticity coordinate of the source color space. */
+	UPROPERTY(EditAnywhere, BlueprintReadWrite, Category = ColorManagement, meta = (EditCondition = "ColorSpaceOverride == ETextureColorSpace::TCS_Custom"))
+	FVector2D RedChromaticityCoordinate;
+
+	/** Green chromaticity coordinate of the source color space. */
+	UPROPERTY(EditAnywhere, BlueprintReadWrite, Category = ColorManagement, meta = (EditCondition = "ColorSpaceOverride == ETextureColorSpace::TCS_Custom"))
+	FVector2D GreenChromaticityCoordinate;
+
+	/** Blue chromaticity coordinate of the source color space. */
+	UPROPERTY(EditAnywhere, BlueprintReadWrite, Category = ColorManagement, meta = (EditCondition = "ColorSpaceOverride == ETextureColorSpace::TCS_Custom"))
+	FVector2D BlueChromaticityCoordinate;
+
+	/** White chromaticity coordinate of the source color space. */
+	UPROPERTY(EditAnywhere, BlueprintReadWrite, Category = ColorManagement, meta = (EditCondition = "ColorSpaceOverride == ETextureColorSpace::TCS_Custom"))
+	FVector2D WhiteChromaticityCoordinate;
+
+	/** Chromatic adaption method applied if the source white point differs from the working color space white point. */
+	UPROPERTY(EditAnywhere, BlueprintReadWrite, Category = ColorManagement, meta = (EditCondition = "ColorSpaceOverride != ETextureColorSpace::TCS_None"))
+	ETextureChromaticAdaptationMethod ChromaticAdaptationMethod;
+
+#if WITH_EDITOR
+	/** Update the chromaticity coordinates member variables based on the color space choice (unless custom). */
+	IMGMEDIA_API void UpdateColorSpaceChromaticities();
+#endif
+};
+
+/* Engine-native color source settings container for media option. */
+struct FNativeMediaSourceColorSettings : public IMediaOptions::FDataContainer
+{
+	/** Constructor */
+	FNativeMediaSourceColorSettings();
+
+	/** Updates the native settings from user-controlled settings. */
+	void Update(const FMediaSourceColorSettings& InSettings);
+
+	/** Color encoding override getter. */
+	UE::Color::EEncoding GetEncodingOverride() const { return EncodingOverride; }
+
+	/** Color space override getter. */
+	const UE::Color::FColorSpace& GetColorSpaceOverride(const UE::Color::FColorSpace& InDefaultColorSpace) const;
+
+	/** Chromatic adaptation getter. */
+	UE::Color::EChromaticAdaptationMethod GetChromaticAdaptationMethod() const { return ChromaticAdaptationMethod; }
+
+private:
+	/** Manual source encoding override. */
+	std::atomic<UE::Color::EEncoding> EncodingOverride;
+
+	/** Manual source color space override. */
+	TOptional<UE::Color::FColorSpace> ColorSpaceOverride;
+
+	/** Chromatic adapation to be used on manual source color space override. */
+	std::atomic<UE::Color::EChromaticAdaptationMethod> ChromaticAdaptationMethod;
+
+	/** Protects color space override variable. */
+	mutable FCriticalSection ColorSpaceCriticalSection;
+};
+
 /**
  * Media source for EXR image sequences.
  *
@@ -79,6 +178,10 @@ public:
 	UPROPERTY(EditAnywhere, BlueprintReadOnly, Category="Timecode")
 	FTimecode StartTimecode;
 
+	/** Manual definition of media source color space & encoding. */
+	UPROPERTY(BlueprintReadOnly, EditAnywhere, Category = "Sequence", meta = (DisplayAfter = "SequencePath", ShowOnlyInnerProperties))
+	FMediaSourceColorSettings SourceColorSettings;
+
 public:
 
 	/**
@@ -195,4 +298,7 @@ protected:
 
 	/** MipMapInfo object to handle mip maps. */
 	TSharedPtr<FImgMediaMipMapInfo, ESPMode::ThreadSafe> MipMapInfo;
+
+	/** Native source color settings. */
+	TSharedPtr<FNativeMediaSourceColorSettings, ESPMode::ThreadSafe> NativeSourceColorSettings;
 };
diff --git a/Engine/Source/Runtime/Media/Public/IMediaTextureSample.h b/Engine/Source/Runtime/Media/Public/IMediaTextureSample.h
--- a/Engine/Source/Runtime/Media/Public/IMediaTextureSample.h
+++ b/Engine/Source/Runtime/Media/Public/IMediaTextureSample.h
@@ -495,6 +495,14 @@ public:
 		return FVector2d(0.15, 0.06);
 	}
 
+	/**
+	* Chromatic adaptation method to be used when applying a color space transform (i.e. from source to working color space).
+	*/
+	virtual UE::Color::EChromaticAdaptationMethod GetChromaticAdapationMethod() const
+	{
+		return UE::Color::EChromaticAdaptationMethod::Bradford;
+	}
+
 	/**
 	 * Get EOTF / "Gamma" / encoding type of data
 	 */
diff --git a/Engine/Source/Runtime/MediaAssets/Private/Misc/MediaTextureResource.cpp b/Engine/Source/Runtime/MediaAssets/Private/Misc/MediaTextureResource.cpp
--- a/Engine/Source/Runtime/MediaAssets/Private/Misc/MediaTextureResource.cpp
+++ b/Engine/Source/Runtime/MediaAssets/Private/Misc/MediaTextureResource.cpp
@@ -1062,7 +1062,7 @@ void FMediaTextureResource::GetColorSpaceConversionMatrixForSample(const TShared
 	else
 	{
 		// Apply the color space transformation from source to destination (including the Bradford chromatic adaptation).
-		ColorSpaceMtx = UE::Color::Transpose<float>(UE::Color::FColorSpaceTransform(Sample->GetSourceColorSpace(), DestinationCS));
+		ColorSpaceMtx = UE::Color::Transpose<float>(UE::Color::FColorSpaceTransform(Sample->GetSourceColorSpace(), DestinationCS, Sample->GetChromaticAdapationMethod()));
 	}
 	
 	float NF = Sample->GetHDRNitsNormalizationFactor();
