ImgMedia: Support color space & encoding overrides in uncompressed EXR gpu path.

* We also move recently added source color settings struct to their own file to avoid include pollution.
* A shader new permutation doesn't appear to be worth it given cost to the default disabled branch (3K car plate: 0.06004ms before, 0.06070ms with the branch disabled, 0.07397 with active matrix + PQ EOTF).

diff --git a/Engine/Plugins/Media/ImgMedia/Shaders/Private/ExrSwizzler.usf b/Engine/Plugins/Media/ImgMedia/Shaders/Private/ExrSwizzler.usf
--- a/Engine/Plugins/Media/ImgMedia/Shaders/Private/ExrSwizzler.usf
+++ b/Engine/Plugins/Media/ImgMedia/Shaders/Private/ExrSwizzler.usf
@@ -1,6 +1,7 @@
 // Copyright Epic Games, Inc. All Rights Reserved.
 
 #include "/Engine/Private/Common.ush"
+#include "/Engine/Private/GammaCorrectionCommon.ush"
 
 #define NUM_CHANNELS (PERMUTATION_CHANNELS + 1)
 
@@ -33,6 +34,36 @@ int2 NumTiles;
 /** Number of channels per pixel. */
 int NumChannels;
 
+/** Flag to enable color space transformations. */
+uint bApplyColorTransform;
+
+/** Source encoding to linearize. */
+uint EOTF;
+
+/** Color space transformation matrix, from source to working. */
+float4x4 ColorSpaceMatrix;
+
+// Same definitions as MediaShaders.usf
+#define UE_Color_EEncoding_Linear		1u
+#define UE_Color_EEncoding_sRGB			2u
+#define UE_Color_EEncoding_ST2084		3u
+#define UE_Color_EEncoding_SLog3		12u
+
+float3 ApplyColorTransform(float3 Color, float3x3 InColorSpaceMatrix, uint InEOTF)
+{
+	// Linearize
+	switch(InEOTF)
+	{
+	case UE_Color_EEncoding_sRGB:	Color = sRGBToLinear(Color); break;
+	case UE_Color_EEncoding_ST2084: Color = ST2084ToLinear(Color); break;
+	case UE_Color_EEncoding_SLog3:	Color = SLog3ToLinear(Color); break;
+	default: break;
+	}
+
+	// Adjust colorspace
+	return mul(InColorSpaceMatrix, Color);
+}
+
 /** Unpacks data from the provided structured buffer. */
 half GetBufferValue(int Offset, int2 TextureSize, int ChannelIndex)
 {
@@ -154,5 +185,10 @@ half4 MainPS(noperspective float4 UVAndScreenPos : TEXCOORD0) : SV_Target0
 	Color.rgb = Color.r;
 #endif
 
+	if(bApplyColorTransform)
+	{
+		Color.rgb = ApplyColorTransform(Color.rgb, (float3x3)ColorSpaceMatrix, EOTF);
+	}
+
 	return Color;
 }
\ No newline at end of file
diff --git a/Engine/Plugins/Media/ImgMedia/Source/ExrReaderGpu/Public/ExrSwizzlingShader.h b/Engine/Plugins/Media/ImgMedia/Source/ExrReaderGpu/Public/ExrSwizzlingShader.h
--- a/Engine/Plugins/Media/ImgMedia/Source/ExrReaderGpu/Public/ExrSwizzlingShader.h
+++ b/Engine/Plugins/Media/ImgMedia/Source/ExrReaderGpu/Public/ExrSwizzlingShader.h
@@ -40,6 +40,9 @@ class FExrSwizzlePS : public FGlobalShader
 		SHADER_PARAMETER(FIntPoint, TileSize)
 		SHADER_PARAMETER(FIntPoint, NumTiles)
 		SHADER_PARAMETER(int32, NumChannels)
+		SHADER_PARAMETER(uint32, bApplyColorTransform)
+		SHADER_PARAMETER(uint32, EOTF)
+		SHADER_PARAMETER(FMatrix44f, ColorSpaceMatrix)
 	END_SHADER_PARAMETER_STRUCT()
 };
 #endif
\ No newline at end of file
diff --git a/Engine/Plugins/Media/ImgMedia/Source/ImgMedia/Internal/Loader/ImgMediaLoader.h b/Engine/Plugins/Media/ImgMedia/Source/ImgMedia/Internal/Loader/ImgMediaLoader.h
--- a/Engine/Plugins/Media/ImgMedia/Source/ImgMedia/Internal/Loader/ImgMediaLoader.h
+++ b/Engine/Plugins/Media/ImgMedia/Source/ImgMedia/Internal/Loader/ImgMediaLoader.h
@@ -391,6 +391,11 @@ public:
 	/** Query if time represents last last frame or later */
 	bool IsFrameLast(const FTimespan& TimeStamp) const;
 
+	/**
+	 * Source color settings getter.
+	 */
+	TSharedPtr<FNativeMediaSourceColorSettings, ESPMode::ThreadSafe> GetSourceColorSettings() const { return SourceColorSettings; }
+
 protected:
 	struct FFrameNumberJobInfo
 	{
diff --git a/Engine/Plugins/Media/ImgMedia/Source/ImgMedia/Private/Assets/ImgMediaSource.cpp b/Engine/Plugins/Media/ImgMedia/Source/ImgMedia/Private/Assets/ImgMediaSource.cpp
--- a/Engine/Plugins/Media/ImgMedia/Source/ImgMedia/Private/Assets/ImgMediaSource.cpp
+++ b/Engine/Plugins/Media/ImgMedia/Source/ImgMedia/Private/Assets/ImgMediaSource.cpp
@@ -14,63 +14,6 @@
 
 #include UE_INLINE_GENERATED_CPP_BY_NAME(ImgMediaSource)
 
-#if WITH_EDITOR
-void FMediaSourceColorSettings::UpdateColorSpaceChromaticities()
-{
-	if (ColorSpaceOverride != ETextureColorSpace::TCS_Custom)
-	{
-		UE::Color::FColorSpace Chromaticities(static_cast<UE::Color::EColorSpace>(ColorSpaceOverride));
-		Chromaticities.GetChromaticities(RedChromaticityCoordinate, GreenChromaticityCoordinate, BlueChromaticityCoordinate, WhiteChromaticityCoordinate);
-	}
-}
-#endif // WITH_EDITOR
-
-FNativeMediaSourceColorSettings::FNativeMediaSourceColorSettings()
-	: EncodingOverride(UE::Color::EEncoding::None)
-	, ColorSpaceOverride()
-	, ChromaticAdaptationMethod(UE::Color::EChromaticAdaptationMethod::Bradford)
-{}
-
-void FNativeMediaSourceColorSettings::Update(const FMediaSourceColorSettings& InSettings)
-{
-	{
-		FScopeLock Lock(&ColorSpaceCriticalSection);
-
-		if (InSettings.ColorSpaceOverride == ETextureColorSpace::TCS_Custom)
-		{
-			ColorSpaceOverride = UE::Color::FColorSpace(
-				InSettings.RedChromaticityCoordinate,
-				InSettings.GreenChromaticityCoordinate,
-				InSettings.BlueChromaticityCoordinate,
-				InSettings.WhiteChromaticityCoordinate
-			);
-		}
-		else if (InSettings.ColorSpaceOverride != ETextureColorSpace::TCS_None)
-		{
-			ColorSpaceOverride = UE::Color::FColorSpace(static_cast<UE::Color::EColorSpace>(InSettings.ColorSpaceOverride));
-		}
-		else
-		{
-			ColorSpaceOverride.Reset();
-		}
-	}
-
-	EncodingOverride = static_cast<UE::Color::EEncoding>(InSettings.EncodingOverride);
-	ChromaticAdaptationMethod = static_cast<UE::Color::EChromaticAdaptationMethod>(InSettings.ChromaticAdaptationMethod);
-}
-
-const UE::Color::FColorSpace& FNativeMediaSourceColorSettings::GetColorSpaceOverride(const UE::Color::FColorSpace& InDefaultColorSpace) const
-{
-	FScopeLock Lock(&ColorSpaceCriticalSection);
-	
-	if (ColorSpaceOverride.IsSet())
-	{
-		return ColorSpaceOverride.GetValue();
-	}
-
-	return InDefaultColorSpace;
-}
-
 /* UImgMediaSource structors
  *****************************************************************************/
 
diff --git a/Engine/Plugins/Media/ImgMedia/Source/ImgMedia/Private/Assets/ImgMediaSourceColorSettings.cpp b/Engine/Plugins/Media/ImgMedia/Source/ImgMedia/Private/Assets/ImgMediaSourceColorSettings.cpp
new file mode 100644
--- /dev/null
+++ b/Engine/Plugins/Media/ImgMedia/Source/ImgMedia/Private/Assets/ImgMediaSourceColorSettings.cpp
@@ -0,0 +1,61 @@
+// Copyright Epic Games, Inc. All Rights Reserved.
+
+#include "ImgMediaSourceColorSettings.h"
+
+#if WITH_EDITOR
+void FMediaSourceColorSettings::UpdateColorSpaceChromaticities()
+{
+	if (ColorSpaceOverride != ETextureColorSpace::TCS_Custom)
+	{
+		UE::Color::FColorSpace Chromaticities(static_cast<UE::Color::EColorSpace>(ColorSpaceOverride));
+		Chromaticities.GetChromaticities(RedChromaticityCoordinate, GreenChromaticityCoordinate, BlueChromaticityCoordinate, WhiteChromaticityCoordinate);
+	}
+}
+#endif // WITH_EDITOR
+
+FNativeMediaSourceColorSettings::FNativeMediaSourceColorSettings()
+	: EncodingOverride(UE::Color::EEncoding::None)
+	, ColorSpaceOverride()
+	, ChromaticAdaptationMethod(UE::Color::EChromaticAdaptationMethod::Bradford)
+{}
+
+void FNativeMediaSourceColorSettings::Update(const FMediaSourceColorSettings& InSettings)
+{
+	{
+		FScopeLock Lock(&ColorSpaceCriticalSection);
+
+		if (InSettings.ColorSpaceOverride == ETextureColorSpace::TCS_Custom)
+		{
+			ColorSpaceOverride = UE::Color::FColorSpace(
+				InSettings.RedChromaticityCoordinate,
+				InSettings.GreenChromaticityCoordinate,
+				InSettings.BlueChromaticityCoordinate,
+				InSettings.WhiteChromaticityCoordinate
+			);
+		}
+		else if (InSettings.ColorSpaceOverride != ETextureColorSpace::TCS_None)
+		{
+			ColorSpaceOverride = UE::Color::FColorSpace(static_cast<UE::Color::EColorSpace>(InSettings.ColorSpaceOverride));
+		}
+		else
+		{
+			ColorSpaceOverride.Reset();
+		}
+	}
+
+	EncodingOverride = static_cast<UE::Color::EEncoding>(InSettings.EncodingOverride);
+	ChromaticAdaptationMethod = static_cast<UE::Color::EChromaticAdaptationMethod>(InSettings.ChromaticAdaptationMethod);
+}
+
+const UE::Color::FColorSpace& FNativeMediaSourceColorSettings::GetColorSpaceOverride(const UE::Color::FColorSpace& InDefaultColorSpace) const
+{
+	FScopeLock Lock(&ColorSpaceCriticalSection);
+	
+	if (ColorSpaceOverride.IsSet())
+	{
+		return ColorSpaceOverride.GetValue();
+	}
+
+	return InDefaultColorSpace;
+}
+
diff --git a/Engine/Plugins/Media/ImgMedia/Source/ImgMedia/Private/Loader/ImgMediaLoader.cpp b/Engine/Plugins/Media/ImgMedia/Source/ImgMedia/Private/Loader/ImgMediaLoader.cpp
--- a/Engine/Plugins/Media/ImgMedia/Source/ImgMedia/Private/Loader/ImgMediaLoader.cpp
+++ b/Engine/Plugins/Media/ImgMedia/Source/ImgMedia/Private/Loader/ImgMediaLoader.cpp
@@ -22,7 +22,7 @@
 #include "Readers/IImgMediaReader.h"
 #include "Readers/DdsImgMediaReader.h"
 #include "ImgMediaLoaderWork.h"
-#include "ImgMediaSource.h"
+#include "ImgMediaSourceColorSettings.h"
 #include "Scheduler/ImgMediaScheduler.h"
 #include "Player/ImgMediaTextureSample.h"
 
diff --git a/Engine/Plugins/Media/ImgMedia/Source/ImgMedia/Private/Player/ImgMediaTextureSample.h b/Engine/Plugins/Media/ImgMedia/Source/ImgMedia/Private/Player/ImgMediaTextureSample.h
--- a/Engine/Plugins/Media/ImgMedia/Source/ImgMedia/Private/Player/ImgMediaTextureSample.h
+++ b/Engine/Plugins/Media/ImgMedia/Source/ImgMedia/Private/Player/ImgMediaTextureSample.h
@@ -8,7 +8,7 @@
 #include "Misc/Timespan.h"
 #include "Templates/SharedPointer.h"
 
-#include "ImgMediaSource.h"
+#include "ImgMediaSourceColorSettings.h"
 #include "Readers/IImgMediaReader.h"
 
 
diff --git a/Engine/Plugins/Media/ImgMedia/Source/ImgMedia/Private/Readers/ExrImgMediaReaderGpu.cpp b/Engine/Plugins/Media/ImgMedia/Source/ImgMedia/Private/Readers/ExrImgMediaReaderGpu.cpp
--- a/Engine/Plugins/Media/ImgMedia/Source/ImgMedia/Private/Readers/ExrImgMediaReaderGpu.cpp
+++ b/Engine/Plugins/Media/ImgMedia/Source/ImgMedia/Private/Readers/ExrImgMediaReaderGpu.cpp
@@ -18,6 +18,7 @@
 #include "Engine/TextureRenderTarget2D.h"
 #include "GlobalShader.h"
 #include "CommonRenderResources.h"
+#include "ImgMediaSourceColorSettings.h"
 #include "Loader/ImgMediaLoader.h"
 #include "PostProcess/DrawRectangle.h"
 #include "RHIStaticStates.h"
@@ -245,6 +246,7 @@ bool FExrImgMediaReaderGpu::ReadFrame(int32 FrameId, const TMap<int32, FImgMedia
 	ConverterParams.TileDimWithBorders = FrameInfo.TileDimensions + FrameInfo.TileBorder * 2;
 	ConverterParams.NumMipLevels = Loader->GetNumMipLevels();
 	ConverterParams.bMipsInSeparateFiles = Loader->MipsInSeparateFiles();
+	ConverterParams.SourceColorSettings = Loader->GetSourceColorSettings();
 
 	{
 		// Force mip level to be upscaled to all higher quality mips.
@@ -499,6 +501,30 @@ void FExrImgMediaReaderGpu::CreateSampleConverterCallback(TSharedPtr<FExrMediaTe
 			{
 				Parameters.NumTiles = FIntPoint(FMath::CeilToInt(float(SampleSize.X) / ConverterParams.TileDimWithBorders.X), FMath::CeilToInt(float(SampleSize.Y) / ConverterParams.TileDimWithBorders.Y));
 			}
+			if (ConverterParams.SourceColorSettings.IsValid())
+			{
+				Parameters.bApplyColorTransform = 1u;
+				Parameters.EOTF = static_cast<uint32>(ConverterParams.SourceColorSettings->GetEncodingOverride());
+
+				const UE::Color::FColorSpace& DestinationCS = UE::Color::FColorSpace::GetWorking();
+				const UE::Color::FColorSpace& SourceCS = ConverterParams.SourceColorSettings->GetColorSpaceOverride(DestinationCS);
+				
+				if (SourceCS.Equals(DestinationCS))
+				{
+					Parameters.ColorSpaceMatrix = FMatrix44f::Identity;
+				}
+				else
+				{
+					const UE::Color::EChromaticAdaptationMethod Method = ConverterParams.SourceColorSettings->GetChromaticAdaptationMethod();
+					Parameters.ColorSpaceMatrix = UE::Color::Transpose<float>(UE::Color::FColorSpaceTransform(SourceCS, DestinationCS, Method));
+				}
+			}
+			else
+			{
+				Parameters.bApplyColorTransform = 0u;
+				Parameters.EOTF = 0u;
+				Parameters.ColorSpaceMatrix = FMatrix44f::Identity;
+			}
 
 			if (ConverterParams.FrameInfo.bHasTiles &&
 				(ConverterParams.TileInfoPerMipLevel.Num() > SampleMipLevel && ConverterParams.TileInfoPerMipLevel[SampleMipLevel].Num() > 0))
diff --git a/Engine/Plugins/Media/ImgMedia/Source/ImgMedia/Private/Readers/ExrImgMediaReaderGpu.h b/Engine/Plugins/Media/ImgMedia/Source/ImgMedia/Private/Readers/ExrImgMediaReaderGpu.h
--- a/Engine/Plugins/Media/ImgMedia/Source/ImgMedia/Private/Readers/ExrImgMediaReaderGpu.h
+++ b/Engine/Plugins/Media/ImgMedia/Source/ImgMedia/Private/Readers/ExrImgMediaReaderGpu.h
@@ -12,6 +12,7 @@
 
 class FExrImgMediaReaderGpu;
 class FExrMediaTextureSampleConverter;
+struct FNativeMediaSourceColorSettings;
 
 struct FStructuredBufferPoolItem
 {
@@ -173,6 +174,9 @@ struct FSampleConverterParameters
 
 	/** A lower quality mip will be upscaled if value is 0 or above. At 0 highest quality mip will always be read fully. */
 	int32 UpscaleMip = -1;
+
+	/** Manual source color space & encoding overrides. */
+	TSharedPtr<FNativeMediaSourceColorSettings, ESPMode::ThreadSafe> SourceColorSettings;
 };
 
 
diff --git a/Engine/Plugins/Media/ImgMedia/Source/ImgMedia/Public/ImgMediaSource.h b/Engine/Plugins/Media/ImgMedia/Source/ImgMedia/Public/ImgMediaSource.h
--- a/Engine/Plugins/Media/ImgMedia/Source/ImgMedia/Public/ImgMediaSource.h
+++ b/Engine/Plugins/Media/ImgMedia/Source/ImgMedia/Public/ImgMediaSource.h
@@ -3,9 +3,7 @@
 #pragma once
 
 #include "CoreTypes.h"
-#include "ColorManagement/ColorSpace.h"
 #include "Engine/EngineTypes.h"
-#include "Engine/TextureDefines.h"
 #include "Containers/UnrealString.h"
 #include "Misc/FrameRate.h"
 #include "Misc/Timecode.h"
@@ -13,6 +11,8 @@
 #include "UObject/ScriptMacros.h"
 #include "BaseMediaSource.h"
 
+#include "ImgMediaSourceColorSettings.h"
+
 #include "ImgMediaSource.generated.h"
 
 class AActor;
@@ -25,102 +25,6 @@ struct IMGMEDIA_API FImgMediaSourceCustomizationSequenceProxy
 	GENERATED_BODY()
 };
 
-// TODO: EMediaSourceEncoding, FMediaSourceColorSettings & FNativeMediaSourceColorSettings should be moved up to BaseMediaSource.h and unified with UCaptureCardMediaSource settings.
-
-/** List of source encodings that can be converted to linear. (Integer values match the ETextureSourceEncoding values in TextureDefines.h */
-UENUM()
-enum class EMediaSourceEncoding : uint8
-{
-	MSE_None	= 0 UMETA(DisplayName = "None", ToolTip = "The source encoding is not overridden."),
-	MSE_Linear	= 1 UMETA(DisplayName = "Linear", ToolTip = "The source encoding is considered linear."),
-	MSE_sRGB	= 2 UMETA(DisplayName = "sRGB", ToolTip = "sRGB source encoding to be linearized"),
-	MSE_ST2084	= 3 UMETA(DisplayName = "ST 2084/PQ", ToolTip = "SMPTE ST 2084/PQ source encoding to be linearized"),
-	MSE_SLog3	= 12 UMETA(DisplayName = "SLog3", ToolTip = "Sony SLog3 source encoding to be linearized"),
-
-	MSE_MAX,
-};
-
-/* Manual definition of media source color space & encoding. */
-USTRUCT(BlueprintType)
-struct FMediaSourceColorSettings
-{
-	GENERATED_USTRUCT_BODY()
-
-	FMediaSourceColorSettings()
-		: EncodingOverride(EMediaSourceEncoding::MSE_None)
-		, ColorSpaceOverride(ETextureColorSpace::TCS_None)
-		, RedChromaticityCoordinate(FVector2D::ZeroVector)
-		, GreenChromaticityCoordinate(FVector2D::ZeroVector)
-		, BlueChromaticityCoordinate(FVector2D::ZeroVector)
-		, WhiteChromaticityCoordinate(FVector2D::ZeroVector)
-		, ChromaticAdaptationMethod(ETextureChromaticAdaptationMethod::TCAM_Bradford)
-	{}
-
-	/** Source encoding of the media. */
-	UPROPERTY(EditAnywhere, BlueprintReadWrite, Category = ColorManagement)
-	EMediaSourceEncoding EncodingOverride;
-
-	/** Source color space of the media. */
-	UPROPERTY(EditAnywhere, BlueprintReadWrite, Category = ColorManagement)
-	ETextureColorSpace ColorSpaceOverride;
-
-	/** Red chromaticity coordinate of the source color space. */
-	UPROPERTY(EditAnywhere, BlueprintReadWrite, Category = ColorManagement, meta = (EditCondition = "ColorSpaceOverride == ETextureColorSpace::TCS_Custom"))
-	FVector2D RedChromaticityCoordinate;
-
-	/** Green chromaticity coordinate of the source color space. */
-	UPROPERTY(EditAnywhere, BlueprintReadWrite, Category = ColorManagement, meta = (EditCondition = "ColorSpaceOverride == ETextureColorSpace::TCS_Custom"))
-	FVector2D GreenChromaticityCoordinate;
-
-	/** Blue chromaticity coordinate of the source color space. */
-	UPROPERTY(EditAnywhere, BlueprintReadWrite, Category = ColorManagement, meta = (EditCondition = "ColorSpaceOverride == ETextureColorSpace::TCS_Custom"))
-	FVector2D BlueChromaticityCoordinate;
-
-	/** White chromaticity coordinate of the source color space. */
-	UPROPERTY(EditAnywhere, BlueprintReadWrite, Category = ColorManagement, meta = (EditCondition = "ColorSpaceOverride == ETextureColorSpace::TCS_Custom"))
-	FVector2D WhiteChromaticityCoordinate;
-
-	/** Chromatic adaption method applied if the source white point differs from the working color space white point. */
-	UPROPERTY(EditAnywhere, BlueprintReadWrite, Category = ColorManagement, meta = (EditCondition = "ColorSpaceOverride != ETextureColorSpace::TCS_None"))
-	ETextureChromaticAdaptationMethod ChromaticAdaptationMethod;
-
-#if WITH_EDITOR
-	/** Update the chromaticity coordinates member variables based on the color space choice (unless custom). */
-	IMGMEDIA_API void UpdateColorSpaceChromaticities();
-#endif
-};
-
-/* Engine-native color source settings container for media option. */
-struct FNativeMediaSourceColorSettings : public IMediaOptions::FDataContainer
-{
-	/** Constructor */
-	FNativeMediaSourceColorSettings();
-
-	/** Updates the native settings from user-controlled settings. */
-	void Update(const FMediaSourceColorSettings& InSettings);
-
-	/** Color encoding override getter. */
-	UE::Color::EEncoding GetEncodingOverride() const { return EncodingOverride; }
-
-	/** Color space override getter. */
-	const UE::Color::FColorSpace& GetColorSpaceOverride(const UE::Color::FColorSpace& InDefaultColorSpace) const;
-
-	/** Chromatic adaptation getter. */
-	UE::Color::EChromaticAdaptationMethod GetChromaticAdaptationMethod() const { return ChromaticAdaptationMethod; }
-
-private:
-	/** Manual source encoding override. */
-	std::atomic<UE::Color::EEncoding> EncodingOverride;
-
-	/** Manual source color space override. */
-	TOptional<UE::Color::FColorSpace> ColorSpaceOverride;
-
-	/** Chromatic adapation to be used on manual source color space override. */
-	std::atomic<UE::Color::EChromaticAdaptationMethod> ChromaticAdaptationMethod;
-
-	/** Protects color space override variable. */
-	mutable FCriticalSection ColorSpaceCriticalSection;
-};
 
 /**
  * Media source for EXR image sequences.
diff --git a/Engine/Plugins/Media/ImgMedia/Source/ImgMedia/Public/ImgMediaSourceColorSettings.h b/Engine/Plugins/Media/ImgMedia/Source/ImgMedia/Public/ImgMediaSourceColorSettings.h
new file mode 100644
--- /dev/null
+++ b/Engine/Plugins/Media/ImgMedia/Source/ImgMedia/Public/ImgMediaSourceColorSettings.h
@@ -0,0 +1,107 @@
+// Copyright Epic Games, Inc. All Rights Reserved.
+
+#pragma once
+
+#include "CoreTypes.h"
+#include "ColorManagement/ColorSpace.h"
+#include "Engine/TextureDefines.h"
+#include "IMediaOptions.h"
+
+#include "ImgMediaSourceColorSettings.generated.h"
+
+// TODO: EMediaSourceEncoding, FMediaSourceColorSettings & FNativeMediaSourceColorSettings should be moved up to BaseMediaSource.h and unified with UCaptureCardMediaSource settings.
+
+/** List of source encodings that can be converted to linear. (Integer values match the ETextureSourceEncoding values in TextureDefines.h */
+UENUM()
+enum class EMediaSourceEncoding : uint8
+{
+	MSE_None	= 0 UMETA(DisplayName = "None", ToolTip = "The source encoding is not overridden."),
+	MSE_Linear	= 1 UMETA(DisplayName = "Linear", ToolTip = "The source encoding is considered linear."),
+	MSE_sRGB	= 2 UMETA(DisplayName = "sRGB", ToolTip = "sRGB source encoding to be linearized"),
+	MSE_ST2084	= 3 UMETA(DisplayName = "ST 2084/PQ", ToolTip = "SMPTE ST 2084/PQ source encoding to be linearized"),
+	MSE_SLog3	= 12 UMETA(DisplayName = "SLog3", ToolTip = "Sony SLog3 source encoding to be linearized"),
+
+	MSE_MAX,
+};
+
+/* Manual definition of media source color space & encoding. */
+USTRUCT(BlueprintType)
+struct FMediaSourceColorSettings
+{
+	GENERATED_USTRUCT_BODY()
+
+	FMediaSourceColorSettings()
+		: EncodingOverride(EMediaSourceEncoding::MSE_None)
+		, ColorSpaceOverride(ETextureColorSpace::TCS_None)
+		, RedChromaticityCoordinate(FVector2D::ZeroVector)
+		, GreenChromaticityCoordinate(FVector2D::ZeroVector)
+		, BlueChromaticityCoordinate(FVector2D::ZeroVector)
+		, WhiteChromaticityCoordinate(FVector2D::ZeroVector)
+		, ChromaticAdaptationMethod(ETextureChromaticAdaptationMethod::TCAM_Bradford)
+	{}
+
+	/** Source encoding of the media. */
+	UPROPERTY(EditAnywhere, BlueprintReadWrite, Category = ColorManagement)
+	EMediaSourceEncoding EncodingOverride;
+
+	/** Source color space of the media. */
+	UPROPERTY(EditAnywhere, BlueprintReadWrite, Category = ColorManagement)
+	ETextureColorSpace ColorSpaceOverride;
+
+	/** Red chromaticity coordinate of the source color space. */
+	UPROPERTY(EditAnywhere, BlueprintReadWrite, Category = ColorManagement, meta = (EditCondition = "ColorSpaceOverride == ETextureColorSpace::TCS_Custom"))
+	FVector2D RedChromaticityCoordinate;
+
+	/** Green chromaticity coordinate of the source color space. */
+	UPROPERTY(EditAnywhere, BlueprintReadWrite, Category = ColorManagement, meta = (EditCondition = "ColorSpaceOverride == ETextureColorSpace::TCS_Custom"))
+	FVector2D GreenChromaticityCoordinate;
+
+	/** Blue chromaticity coordinate of the source color space. */
+	UPROPERTY(EditAnywhere, BlueprintReadWrite, Category = ColorManagement, meta = (EditCondition = "ColorSpaceOverride == ETextureColorSpace::TCS_Custom"))
+	FVector2D BlueChromaticityCoordinate;
+
+	/** White chromaticity coordinate of the source color space. */
+	UPROPERTY(EditAnywhere, BlueprintReadWrite, Category = ColorManagement, meta = (EditCondition = "ColorSpaceOverride == ETextureColorSpace::TCS_Custom"))
+	FVector2D WhiteChromaticityCoordinate;
+
+	/** Chromatic adaption method applied if the source white point differs from the working color space white point. */
+	UPROPERTY(EditAnywhere, BlueprintReadWrite, Category = ColorManagement, meta = (EditCondition = "ColorSpaceOverride != ETextureColorSpace::TCS_None"))
+	ETextureChromaticAdaptationMethod ChromaticAdaptationMethod;
+
+#if WITH_EDITOR
+	/** Update the chromaticity coordinates member variables based on the color space choice (unless custom). */
+	IMGMEDIA_API void UpdateColorSpaceChromaticities();
+#endif
+};
+
+/* Engine-native color source settings container for media option. */
+struct FNativeMediaSourceColorSettings : public IMediaOptions::FDataContainer
+{
+	/** Constructor */
+	FNativeMediaSourceColorSettings();
+
+	/** Updates the native settings from user-controlled settings. */
+	void Update(const FMediaSourceColorSettings& InSettings);
+
+	/** Color encoding override getter. */
+	UE::Color::EEncoding GetEncodingOverride() const { return EncodingOverride; }
+
+	/** Color space override getter. */
+	const UE::Color::FColorSpace& GetColorSpaceOverride(const UE::Color::FColorSpace& InDefaultColorSpace) const;
+
+	/** Chromatic adaptation getter. */
+	UE::Color::EChromaticAdaptationMethod GetChromaticAdaptationMethod() const { return ChromaticAdaptationMethod; }
+
+private:
+	/** Manual source encoding override. */
+	std::atomic<UE::Color::EEncoding> EncodingOverride;
+
+	/** Manual source color space override. */
+	TOptional<UE::Color::FColorSpace> ColorSpaceOverride;
+
+	/** Chromatic adapation to be used on manual source color space override. */
+	std::atomic<UE::Color::EChromaticAdaptationMethod> ChromaticAdaptationMethod;
+
+	/** Protects color space override variable. */
+	mutable FCriticalSection ColorSpaceCriticalSection;
+};
