[MediaPlate] MediaPlateResource Editor: when in external file mode, make the paths relative to content folder (priority) or project folder (fallback).

Consistency pass on all path resolving for media source dealing with file paths.

Both ImgMediaSource and FileMediaSource will support relative paths either to the Content folder (priority) or Project Folder (fallback).

The rules are:
- Try to make the path relative to content (first) or project. If not possible, leave absolute.
- Relative paths begins with "./" as a convention.
- The exception is ImgMediaSource which also supports tokens (project_dir, engine_dir).

diff --git a/Engine/Plugins/Media/ImgMedia/Source/ImgMedia/Private/Assets/ImgMediaSource.cpp b/Engine/Plugins/Media/ImgMedia/Source/ImgMedia/Private/Assets/ImgMediaSource.cpp
--- a/Engine/Plugins/Media/ImgMedia/Source/ImgMedia/Private/Assets/ImgMediaSource.cpp
+++ b/Engine/Plugins/Media/ImgMedia/Source/ImgMedia/Private/Assets/ImgMediaSource.cpp
@@ -14,6 +14,39 @@
 
 #include UE_INLINE_GENERATED_CPP_BY_NAME(ImgMediaSource)
 
+namespace UE::ImgMediaSource::Private
+{
+	/**
+	 * Returns list of possible sequence base paths in order of priority. 
+	 */
+	const TArray<FString>& GetPossibleBasePaths()
+	{
+		static const TArray<FString> PossibleBasePaths
+		{
+			FPaths::ProjectContentDir(),
+			FPaths::ProjectDir()
+		};
+		return PossibleBasePaths;
+	}
+
+	/**
+	 * Given a full path, find under which base path it could be.
+	 * returns an empty string if it is not under any base dir.
+	 */
+	FString FindFullBasePath(const FString& InFullPath)
+	{
+		for (const FString& BasePath : GetPossibleBasePaths())
+		{
+			const FString FullBasePath = FPaths::ConvertRelativePathToFull(BasePath);
+			if (FPaths::IsUnderDirectory(InFullPath, FullBasePath))
+			{
+				return FullBasePath;
+			}
+		}
+		return FString();
+	}
+}
+
 /* UImgMediaSource structors
  *****************************************************************************/
 
@@ -42,25 +75,7 @@ const FString UImgMediaSource::GetSequencePath() const
 
 void UImgMediaSource::SetSequencePath(const FString& Path)
 {
-	const FString SanitizedPath = FPaths::GetPath(Path);
-
-	if (SanitizedPath.IsEmpty() || SanitizedPath.StartsWith(TEXT(".")))
-	{
-		SequencePath.Path = SanitizedPath;
-	}
-	else
-	{
-		FString FullPath = FPaths::ConvertRelativePathToFull(SanitizedPath);
-		const FString RelativeDir = FPaths::ConvertRelativePathToFull(FPaths::ProjectDir());
-
-		if (FullPath.StartsWith(RelativeDir))
-		{
-			FPaths::MakePathRelativeTo(FullPath, *RelativeDir);
-			FullPath = FString(TEXT("./")) + FullPath;
-		}
-
-		SequencePath.Path = FullPath;
-	}
+	SequencePath.Path = SanitizeTokenizedSequencePath(Path);
 }
 
 void UImgMediaSource::SetTokenizedSequencePath(const FString& Path)
@@ -78,39 +93,85 @@ FString UImgMediaSource::ExpandSequencePathTokens(const FString& InPath)
 
 FString UImgMediaSource::SanitizeTokenizedSequencePath(const FString& InPath)
 {
-	FString SanitizedPickedPath = InPath.TrimStartAndEnd().Replace(TEXT("\""), TEXT(""));
+	FString SanitizedPath = InPath.TrimStartAndEnd().Replace(TEXT("\""), TEXT(""));
+	FPaths::NormalizeDirectoryName(SanitizedPath);
 
-	const FString ProjectAbsolutePath = FPaths::ConvertRelativePathToFull(FPaths::ProjectDir());
+	if (SanitizedPath.IsEmpty())
+	{
+		return SanitizedPath;
+	}
 
 	// Replace supported tokens
-	FString ExpandedPath = UImgMediaSource::ExpandSequencePathTokens(SanitizedPickedPath);
-
-	// Relative paths are always w.r.t. the project root.
+	FString ExpandedPath = ExpandSequencePathTokens(SanitizedPath);
+	
+	FString SelectedRootPath;
+	
+	// Expand relative path with the possible base paths.
 	if (FPaths::IsRelative(ExpandedPath))
 	{
-		ExpandedPath = FPaths::Combine(ProjectAbsolutePath, SanitizedPickedPath);
+		// It could be a path relative to the process BaseDir
+		FString FullExpandedPath = FPaths::ConvertRelativePathToFull(ExpandedPath);
+		if (FPaths::DirectoryExists(FullExpandedPath) || FPaths::FileExists(FullExpandedPath))
+		{
+			ExpandedPath = FullExpandedPath;
+			SelectedRootPath = UE::ImgMediaSource::Private::FindFullBasePath(FullExpandedPath);
+		}
+		else
+		{
+			// If it is not relative to BaseDir, we try the other possible bases.
+			for (const FString& BasePath : UE::ImgMediaSource::Private::GetPossibleBasePaths())
+			{
+				const FString FullBasePath = FPaths::ConvertRelativePathToFull(BasePath);
+				FullExpandedPath = FPaths::ConvertRelativePathToFull(FullBasePath, ExpandedPath);
+				// Note: Directory or file needs to exist to figure it out.
+				if (FPaths::DirectoryExists(FullExpandedPath) || FPaths::FileExists(FullExpandedPath))
+				{
+					ExpandedPath = FullExpandedPath;
+					SelectedRootPath = FullBasePath;
+					break;
+				}
+			}
+		}
+	}
+	else
+	{
+		// For an absolute path, we still need to find which base it is under.
+		SelectedRootPath = UE::ImgMediaSource::Private::FindFullBasePath(ExpandedPath);
 	}
 
 	// Chop trailing file path, in case the user picked a file instead of a folder
 	if (FPaths::FileExists(ExpandedPath))
 	{
 		ExpandedPath = FPaths::GetPath(ExpandedPath);
-		SanitizedPickedPath = FPaths::GetPath(SanitizedPickedPath);
+		SanitizedPath = FPaths::GetPath(SanitizedPath);
 	}
 
 	// If the user picked the absolute path of a directory that is inside the project, use relative path.
 	// Unless the user has a token in the beginning.
 	if (!InPath.Len() || InPath[0] != '{') // '{' indicates that the path begins with a token
 	{
-		FString PathRelativeToProject;
+		if (!SelectedRootPath.IsEmpty())
+		{
+			FString PathRelativeToProject;
 
-		if (IsPathUnderBasePath(ExpandedPath, ProjectAbsolutePath, PathRelativeToProject))
+			if (IsPathUnderBasePath(ExpandedPath, SelectedRootPath, PathRelativeToProject))
+			{
+				// Sanitized relative path expected to start with "./"
+				if (!PathRelativeToProject.StartsWith(TEXT("./")))
+				{
+					PathRelativeToProject = FPaths::Combine(TEXT("."), PathRelativeToProject);
+				}
+				SanitizedPath = PathRelativeToProject;
+			}
+		}
+		else
 		{
-			SanitizedPickedPath = PathRelativeToProject;
+			// the path was not inside the project, return absolute path.
+			SanitizedPath = ExpandedPath;
 		}
 	}
 
-	return SanitizedPickedPath;
+	return SanitizedPath;
 }
 
 void UImgMediaSource::AddTargetObject(AActor* InActor)
@@ -258,6 +319,16 @@ FString UImgMediaSource::GetFullPath() const
 
 	if (FPaths::IsRelative(ExpandedSequencePath))
 	{
+		for (const FString& BasePath : UE::ImgMediaSource::Private::GetPossibleBasePaths())
+		{
+			const FString FullPath = FPaths::ConvertRelativePathToFull(FPaths::Combine(BasePath, ExpandedSequencePath));
+			if (FPaths::DirectoryExists(FullPath))
+			{
+				return FullPath;
+			}
+		}
+
+		// If we can't confirm because path doesn't exist, we will default to project dir for backward compatibility.
 		return FPaths::ConvertRelativePathToFull(FPaths::Combine(FPaths::ProjectDir(), ExpandedSequencePath));
 	}
 	else
@@ -283,9 +354,7 @@ void UImgMediaSource::Serialize(FArchive& Ar)
 
 			if (FPaths::IsRelative(SequencePath.Path))
 			{
-				SequencePath.Path = FString::Printf(TEXT("Content/%s"), *SequencePath.Path);
-
-				SequencePath.Path = UImgMediaSource::SanitizeTokenizedSequencePath(SequencePath.Path);
+				SequencePath.Path = SanitizeTokenizedSequencePath(SequencePath.Path);
 			}
 		}
 	}
diff --git a/Engine/Plugins/Media/ImgMedia/Source/ImgMedia/Private/Tests/ImgMediaTests.cpp b/Engine/Plugins/Media/ImgMedia/Source/ImgMedia/Private/Tests/ImgMediaTests.cpp
--- a/Engine/Plugins/Media/ImgMedia/Source/ImgMedia/Private/Tests/ImgMediaTests.cpp
+++ b/Engine/Plugins/Media/ImgMedia/Source/ImgMedia/Private/Tests/ImgMediaTests.cpp
@@ -1,11 +1,13 @@
 // Copyright Epic Games, Inc. All Rights Reserved.
 
-#include "CoreTypes.h"
 #include "Containers/UnrealString.h"
+#include "CoreTypes.h"
+#include "HAL/FileManager.h"
 #include "Logging/LogMacros.h"
 #include "Misc/AutomationTest.h"
 
 #include "Assets/ImgMediaMipMapInfo.h"
+#include "ImgMediaSource.h"
 
 #if WITH_DEV_AUTOMATION_TESTS
 
@@ -89,4 +91,164 @@ bool FImgMediaTests::RunTest(const FString& Parameters)
 	return !HasAnyErrors();
 }
 
+IMPLEMENT_SIMPLE_AUTOMATION_TEST(FImgMediaSanitizePathTests, "System.Plugins.ImgMedia.SanitizePath", EAutomationTestFlags::EditorContext | EAutomationTestFlags::EngineFilter)
+bool FImgMediaSanitizePathTests::RunTest(const FString& Parameters)
+{
+	// Setup the test, needs some existing files in different folders
+	TArray<FString> CreatedDirectories;
+	TArray<FString> CreatedFiles;
+	
+	auto MakeDirectory = [&CreatedDirectories](const FString& InPath) -> bool
+	{
+		if (FPaths::DirectoryExists(InPath))
+		{
+			return true;
+		}
+		if (IFileManager::Get().MakeDirectory(*InPath))
+		{
+			CreatedDirectories.Add(InPath);
+			return true;
+		}
+		return false;
+	};
+
+	auto MakeFile = [&CreatedFiles](const FString& InFilepath) -> bool
+	{
+		if (FPaths::FileExists(InFilepath))
+		{
+			return true;
+		}
+		const TUniquePtr<FArchive> FileWriter(IFileManager::Get().CreateFileWriter(*InFilepath));
+		if (FileWriter)
+		{
+			int32 value = 1234;
+			(*FileWriter) << value;
+			CreatedFiles.Add(InFilepath);
+			return true;
+		}
+		return false;
+	};
+
+	{
+		const FString ContentMoviesDirectory = FPaths::ConvertRelativePathToFull(FPaths::Combine(FPaths::ProjectContentDir(), TEXT("Movies")));
+		const FString ContentSequenceDirectory = FPaths::Combine(ContentMoviesDirectory, TEXT("ImgMediaTestSequence"));
+		const FString ContentSequenceImgPath = FPaths::Combine(ContentSequenceDirectory, TEXT("Image0000.exr"));
+		MakeDirectory(ContentMoviesDirectory);
+		MakeDirectory(ContentSequenceDirectory);
+		MakeFile(ContentSequenceImgPath);
+	}
+	{
+		const FString ProjectMoviesDirectory = FPaths::ConvertRelativePathToFull(FPaths::Combine(FPaths::ProjectDir(), TEXT("Movies")));
+		const FString ProjectSequenceDirectory = FPaths::Combine(ProjectMoviesDirectory, TEXT("ImgMediaTestSequence"));
+		const FString ProjectSequenceImgPath = FPaths::Combine(ProjectSequenceDirectory, TEXT("Image0000.exr"));
+		MakeDirectory(ProjectMoviesDirectory);
+		MakeDirectory(ProjectSequenceDirectory);
+		MakeFile(ProjectSequenceImgPath);
+	}
+	{
+		const FString EngineMoviesDirectory = FPaths::ConvertRelativePathToFull(FPaths::Combine(FPaths::EngineDir(), TEXT("Movies")));
+		const FString EngineSequenceDirectory = FPaths::Combine(EngineMoviesDirectory, TEXT("ImgMediaTestSequence"));
+		const FString EngineSequenceImgPath = FPaths::Combine(EngineSequenceDirectory, TEXT("Image0000.exr"));
+		MakeDirectory(EngineMoviesDirectory);
+		MakeDirectory(EngineSequenceDirectory);
+		MakeFile(EngineSequenceImgPath);
+	}
+	
+	ON_SCOPE_EXIT
+	{
+		for (const FString& File : CreatedFiles)
+		{
+			IFileManager::Get().Delete(*File);
+		}
+		
+		for(const FString& Directory : ReverseIterate(CreatedDirectories))
+		{
+			IFileManager::Get().DeleteDirectory(*Directory);
+		}
+	};
+
+	// ---- Tests Begin here
+
+	auto TestSanitizeTokenizedPath = [this](const FString& InPath, const FString& InExpectedResult, FString& OutError)
+	{
+		const FString SanitizedPath = UImgMediaSource::SanitizeTokenizedSequencePath(InPath);
+		if (SanitizedPath != InExpectedResult)
+		{
+			OutError = FString::Printf(TEXT("SanitizeTokenizedSequencePath failed: Input: \"%s\" Result: \"%s\" Expected \"%s\""), *InPath, *SanitizedPath, *InExpectedResult);
+			return false;
+		}
+		return true;
+	};
+
+	FString Error;
+	
+	// Empty path
+	AddErrorIfFalse(TestSanitizeTokenizedPath(TEXT(""), TEXT(""), Error), Error);
+
+	// Already sanitized
+	AddErrorIfFalse(TestSanitizeTokenizedPath(TEXT("./Movies/ImgMediaTestSequence"), TEXT("./Movies/ImgMediaTestSequence"), Error), Error);
+	AddErrorIfFalse(TestSanitizeTokenizedPath(TEXT("./Movies/ImgMediaTestSequence/"), TEXT("./Movies/ImgMediaTestSequence"), Error), Error);
+	AddErrorIfFalse(TestSanitizeTokenizedPath(TEXT(".\\Movies\\ImgMediaTestSequence"), TEXT("./Movies/ImgMediaTestSequence"), Error), Error);
+	AddErrorIfFalse(TestSanitizeTokenizedPath(TEXT(".\\Movies\\ImgMediaTestSequence\\"), TEXT("./Movies/ImgMediaTestSequence"), Error), Error);
+	AddErrorIfFalse(TestSanitizeTokenizedPath(TEXT("\".\\Movies\\ImgMediaTestSequence\""), TEXT("./Movies/ImgMediaTestSequence"), Error), Error);
+
+	// With the filename (it should remove it)
+	AddErrorIfFalse(TestSanitizeTokenizedPath(TEXT("./Movies/ImgMediaTestSequence/Image0000.exr"), TEXT("./Movies/ImgMediaTestSequence"), Error), Error);
+	AddErrorIfFalse(TestSanitizeTokenizedPath(TEXT(".\\Movies\\ImgMediaTestSequence\\Image0000.exr"), TEXT("./Movies/ImgMediaTestSequence"), Error), Error);
+	AddErrorIfFalse(TestSanitizeTokenizedPath(TEXT("\".\\Movies\\ImgMediaTestSequence\\Image0000.exr\""), TEXT("./Movies/ImgMediaTestSequence"), Error), Error);
+	
+	{
+		// Project relative
+		FString ProjectRelative = FPaths::Combine(FPaths::ProjectDir(), TEXT("Movies/ImgMediaTestSequence"));
+		AddErrorIfFalse(TestSanitizeTokenizedPath(ProjectRelative, TEXT("./Movies/ImgMediaTestSequence"), Error), Error);
+		AddErrorIfFalse(TestSanitizeTokenizedPath(FPaths::ConvertRelativePathToFull(ProjectRelative), TEXT("./Movies/ImgMediaTestSequence"), Error), Error);
+
+		// With filename
+		ProjectRelative = FPaths::Combine(ProjectRelative, TEXT("Image0000.exr"));
+		AddErrorIfFalse(TestSanitizeTokenizedPath(ProjectRelative, TEXT("./Movies/ImgMediaTestSequence"), Error), Error);
+		AddErrorIfFalse(TestSanitizeTokenizedPath(FPaths::ConvertRelativePathToFull(ProjectRelative), TEXT("./Movies/ImgMediaTestSequence"), Error), Error);
+	}
+
+	{
+		// Content relative
+		FString ContentRelative = FPaths::Combine(FPaths::ProjectContentDir(), TEXT("Movies/ImgMediaTestSequence"));
+		AddErrorIfFalse(TestSanitizeTokenizedPath(ContentRelative, TEXT("./Movies/ImgMediaTestSequence"), Error), Error);
+		AddErrorIfFalse(TestSanitizeTokenizedPath(FPaths::ConvertRelativePathToFull(ContentRelative), TEXT("./Movies/ImgMediaTestSequence"), Error), Error);
+
+		// With filename
+		ContentRelative = FPaths::Combine(ContentRelative, TEXT("Image0000.exr"));
+		AddErrorIfFalse(TestSanitizeTokenizedPath(ContentRelative, TEXT("./Movies/ImgMediaTestSequence"), Error), Error);
+		AddErrorIfFalse(TestSanitizeTokenizedPath(FPaths::ConvertRelativePathToFull(ContentRelative), TEXT("./Movies/ImgMediaTestSequence"), Error), Error);
+	}
+
+	{
+		// Outside of project (ex: engine), expect absolute path
+		FString EngineRelative = FPaths::Combine(FPaths::EngineDir(), TEXT("Movies/ImgMediaTestSequence"));
+		FString EngineRelativeFull = FPaths::ConvertRelativePathToFull(EngineRelative);	// expected
+		AddErrorIfFalse(TestSanitizeTokenizedPath(EngineRelative, EngineRelativeFull, Error), Error);
+		AddErrorIfFalse(TestSanitizeTokenizedPath(FPaths::ConvertRelativePathToFull(EngineRelative), EngineRelativeFull, Error), Error);
+
+		// With filename
+		EngineRelative = FPaths::Combine(EngineRelative,  TEXT("Image0000.exr"));
+		AddErrorIfFalse(TestSanitizeTokenizedPath(EngineRelative, EngineRelativeFull, Error), Error);
+		AddErrorIfFalse(TestSanitizeTokenizedPath(FPaths::ConvertRelativePathToFull(EngineRelative), EngineRelativeFull, Error), Error);
+	}
+
+	{
+		// Project relative with Token
+		FString Expected = TEXT("{project_dir}/Movies/ImgMediaTestSequence");
+		AddErrorIfFalse(TestSanitizeTokenizedPath(Expected, Expected, Error), Error);
+		AddErrorIfFalse(TestSanitizeTokenizedPath(TEXT("{project_dir}/Movies/ImgMediaTestSequence/Image0000.exr"), Expected, Error), Error);
+	}
+
+	{
+		// Engine relative with Token
+		FString Expected = TEXT("{engine_dir}/Movies/ImgMediaTestSequence");
+		AddErrorIfFalse(TestSanitizeTokenizedPath(Expected, Expected, Error), Error);
+		AddErrorIfFalse(TestSanitizeTokenizedPath(TEXT("{engine_dir}/Movies/ImgMediaTestSequence/Image0000.exr"), Expected, Error), Error);
+	}
+	
+	return !HasAnyErrors();
+}
+
 #endif
diff --git a/Engine/Plugins/Media/ImgMedia/Source/ImgMediaEditor/Private/Customizations/ImgMediaSourceCustomization.cpp b/Engine/Plugins/Media/ImgMedia/Source/ImgMediaEditor/Private/Customizations/ImgMediaSourceCustomization.cpp
--- a/Engine/Plugins/Media/ImgMedia/Source/ImgMediaEditor/Private/Customizations/ImgMediaSourceCustomization.cpp
+++ b/Engine/Plugins/Media/ImgMedia/Source/ImgMediaEditor/Private/Customizations/ImgMediaSourceCustomization.cpp
@@ -37,6 +37,8 @@ void FImgMediaSourceCustomization::CustomizeHeader(TSharedRef<IPropertyHandle> I
 {
 	PropertyHandle = InPropertyHandle;
 
+	TWeakPtr<FImgMediaSourceCustomization> WeakSelf = StaticCastWeakPtr<FImgMediaSourceCustomization>(AsWeak());
+	
 	// customize 'File' category
 	HeaderRow
 		.NameContent()
@@ -74,14 +76,16 @@ void FImgMediaSourceCustomization::CustomizeHeader(TSharedRef<IPropertyHandle> I
 					.BrowseButtonStyle(FAppStyle::Get(), "HoverHintOnly")
 					.BrowseButtonToolTip(LOCTEXT("SequencePathBrowseButtonToolTip", "Choose a file from this computer"))
 					.DialogReturnsFullPath(true)
-					.BrowseDirectory_Lambda([this]() -> FString
+					.BrowseDirectory_Lambda([WeakSelf]() -> FString
 					{
-						const FString SequencePath = GetSequencePath();
+						const TSharedPtr<FImgMediaSourceCustomization> Self = WeakSelf.Pin();
+						const FString SequencePath = Self ? Self->GetSequenceFullPath() : TEXT("");
 						return SequencePath.IsEmpty() ? (FPaths::ProjectContentDir() / TEXT("Movies")) : SequencePath;
 					})
-					.FilePath_Lambda([this]() -> FString
+					.FilePath_Lambda([WeakSelf]() -> FString
 					{
-						return GetSequencePath();
+						const TSharedPtr<FImgMediaSourceCustomization> Self = WeakSelf.Pin();
+						return Self ? Self->GetSequencePath() : TEXT("");
 					})
 					.FileTypeFilter_Lambda([]() -> FString
 					{
@@ -119,6 +123,21 @@ FString FImgMediaSourceCustomization::GetSequencePath() const
 	return GetSequencePathFromChildProperty(PropertyHandle);
 }
 
+FString FImgMediaSourceCustomization::GetSequenceFullPath() const
+{
+	TArray<UObject*> OuterObjects;
+	PropertyHandle->GetOuterObjects(OuterObjects);
+	for (UObject* OuterObject : OuterObjects)
+	{
+		if (const UImgMediaSource* Source = Cast<UImgMediaSource>(OuterObject))
+		{
+			// Using the media source to resolve the full path.
+			return Source->GetFullPath();
+		}
+	}
+	return FString();
+}
+
 TSharedPtr<IPropertyHandle> FImgMediaSourceCustomization::GetSequencePathProperty(const TSharedPtr<IPropertyHandle>& InPropertyHandle)
 {
 	TSharedPtr<IPropertyHandle> SequencePathProperty;
@@ -153,53 +172,11 @@ TSharedPtr<IPropertyHandle> FImgMediaSourceCustomization::GetSequencePathPathPro
 
 void FImgMediaSourceCustomization::HandleSequencePathPickerPathPicked(const FString& PickedPath)
 {
-	FString SanitizedPickedPath = PickedPath.TrimStartAndEnd().Replace(TEXT("\""), TEXT(""));
-
 	// The user may have put a path that is relative to the project, or to the content, or to the running process.
-	{
-		const FString ExpandedPath = UImgMediaSource::ExpandSequencePathTokens(SanitizedPickedPath);
-
-		// We only need to interpret relative path (not absolute paths).
-		// Note: The path picker is set to only returns absolute paths.
-		if (FPaths::IsRelative(ExpandedPath))
-		{
-			// See if the path is relative to project root or content.
-			// Pick the one that exists
-
-			FString SanitizedPickedPathReplacement = SanitizedPickedPath;
-
-			const TArray<TPair<FString, FString>> PossibleBasePaths
-			{
-				TPair<FString,FString>(FPaths::ProjectDir(), TEXT("")),
-				TPair<FString,FString>(FPaths::ProjectContentDir(), TEXT("Content")),
-			};
-
-			int32 ExistedCnt = 0;
-
-			for (const auto& Pair : PossibleBasePaths)
-			{
-				const FString CandidatePath = FPaths::Combine(Pair.Key, ExpandedPath);
-
-				if (FPaths::FileExists(CandidatePath) || FPaths::DirectoryExists(CandidatePath))
-				{
-					SanitizedPickedPathReplacement = FPaths::Combine(Pair.Value, SanitizedPickedPath);
-					ExistedCnt++;
-				}
-			}
-
-			if (ExistedCnt == 1)
-			{
-				// Re-interpret the base relative if the path was relative to either project or content (otherwise leave unaltered)
-				SanitizedPickedPath = SanitizedPickedPathReplacement;
-			}
-		}
-	}
-
-	// Sanitize the path
-	SanitizedPickedPath = UImgMediaSource::SanitizeTokenizedSequencePath(SanitizedPickedPath);
+	const FString SanitizedPickedPath = UImgMediaSource::SanitizeTokenizedSequencePath(PickedPath);
 
 	// update property
-	TSharedPtr<IPropertyHandle> SequencePathPathProperty = GetSequencePathPathProperty(PropertyHandle);
+	const TSharedPtr<IPropertyHandle> SequencePathPathProperty = GetSequencePathPathProperty(PropertyHandle);
 	if (SequencePathPathProperty.IsValid())
 	{
 		if (SequencePathPathProperty->SetValue(SanitizedPickedPath) != FPropertyAccess::Success)
@@ -209,32 +186,28 @@ void FImgMediaSourceCustomization::HandleSequencePathPickerPathPicked(const FStr
 	}
 }
 
-
 EVisibility FImgMediaSourceCustomization::HandleSequencePathWarningIconVisibility() const
 {
-	const FString FilePath = UImgMediaSource::ExpandSequencePathTokens(GetSequencePath());
+	const FString SequenceFullPath = GetSequenceFullPath();
 
-	if (FilePath.IsEmpty() || FilePath.Contains(TEXT("://")))
+	if (SequenceFullPath.IsEmpty())
 	{
 		return EVisibility::Hidden;
 	}
 
 	const FString FullMoviesPath = FPaths::ConvertRelativePathToFull(FPaths::ProjectContentDir() / TEXT("Movies"));
-	const FString FullPath = FPaths::ConvertRelativePathToFull(FPaths::IsRelative(FilePath) ? FPaths::ProjectDir() / FilePath : FilePath);
-	FString RelativePath;
-
-	if (UImgMediaSource::IsPathUnderBasePath(FullPath, FullMoviesPath, RelativePath))
+	if (FPaths::IsUnderDirectory(SequenceFullPath, FullMoviesPath))
 	{
-		if (FPaths::DirectoryExists(FullPath))
+		if (FPaths::DirectoryExists(SequenceFullPath))
 		{
 			return EVisibility::Hidden;
 		}
 
-		// file doesn't exist
+		// Path doesn't exist
 		return EVisibility::Visible;
 	}
 
-	// file not inside Movies folder
+	// Sequence path is not inside Movies folder
 	return EVisibility::Visible;
 }
 
diff --git a/Engine/Plugins/Media/ImgMedia/Source/ImgMediaEditor/Private/Customizations/ImgMediaSourceCustomization.h b/Engine/Plugins/Media/ImgMedia/Source/ImgMediaEditor/Private/Customizations/ImgMediaSourceCustomization.h
--- a/Engine/Plugins/Media/ImgMedia/Source/ImgMediaEditor/Private/Customizations/ImgMediaSourceCustomization.h
+++ b/Engine/Plugins/Media/ImgMedia/Source/ImgMediaEditor/Private/Customizations/ImgMediaSourceCustomization.h
@@ -56,6 +56,11 @@ protected:
 	 */
 	FString GetSequencePath() const;
 
+	/**
+	 *  Returns the fully resolved path to the sequence directory of the currently selected media source.
+	 */
+	FString GetSequenceFullPath() const;
+
 private:
 
 	/** Callback for picking a path in the source directory picker. */
diff --git a/Engine/Plugins/Media/MediaPlate/Source/MediaPlate/Private/MediaPlateComponent.cpp b/Engine/Plugins/Media/MediaPlate/Source/MediaPlate/Private/MediaPlateComponent.cpp
--- a/Engine/Plugins/Media/MediaPlate/Source/MediaPlate/Private/MediaPlateComponent.cpp
+++ b/Engine/Plugins/Media/MediaPlate/Source/MediaPlate/Private/MediaPlateComponent.cpp
@@ -119,21 +119,29 @@ namespace UE::MediaPlateComponent
 			}
 		}
 
-		// At least support the case if the path is within the project (otherwise, leave absolute)
-		FString RelativeToProject = MediaPath;
-		if (FPaths::MakePathRelativeTo(RelativeToProject, *FPaths::ProjectDir()))
+		// Try to make relative under content or project, if not leave absolute.
+		FString ConvertedPath = FPaths::ConvertRelativePathToFull(MediaPath);
+
+		const TArray<FString> PossibleBasePaths =
 		{
-			if (!RelativeToProject.StartsWith(TEXT(".")))
-			{
-				// Project relative paths are expected to start with "./"
-				MediaPath = FPaths::Combine(TEXT("."), RelativeToProject);
-			}
-			else
+			FPaths::ProjectContentDir(),
+			FPaths::ProjectDir()
+		};
+
+		for (const FString& BasePath : PossibleBasePaths)
+		{
+			const FString FullBasePath = FPaths::ConvertRelativePathToFull(BasePath);
+			if (FPaths::IsUnderDirectory(ConvertedPath, FullBasePath) && FPaths::MakePathRelativeTo(ConvertedPath, *FullBasePath))
 			{
-				MediaPath = RelativeToProject;
+				// Relative path expected to start with "./"
+				if (!ConvertedPath.StartsWith(TEXT("./")))
+				{
+					ConvertedPath = FPaths::Combine(TEXT("."), ConvertedPath);
+				}
+				break;
 			}
 		}
-		return MediaPath;
+		return ConvertedPath;
 	}
 #endif
 };
diff --git a/Engine/Plugins/Media/MediaPlate/Source/MediaPlateEditor/Private/MediaPlateResourceCustomization.cpp b/Engine/Plugins/Media/MediaPlate/Source/MediaPlateEditor/Private/MediaPlateResourceCustomization.cpp
--- a/Engine/Plugins/Media/MediaPlate/Source/MediaPlateEditor/Private/MediaPlateResourceCustomization.cpp
+++ b/Engine/Plugins/Media/MediaPlate/Source/MediaPlateEditor/Private/MediaPlateResourceCustomization.cpp
@@ -3,6 +3,7 @@
 #include "MediaPlateResourceCustomization.h"
 #include "DetailLayoutBuilder.h"
 #include "DetailWidgetRow.h"
+#include "EditorDirectories.h"
 #include "MediaPlateComponent.h"
 #include "MediaPlaylist.h"
 #include "MediaSource.h"
@@ -88,6 +89,7 @@ void FMediaPlateResourceCustomization::CustomizeHeader(TSharedRef<IPropertyHandl
 					.BrowseButtonImage(FAppStyle::GetBrush("PropertyWindow.Button_Ellipsis"))
 					.BrowseButtonStyle(FAppStyle::Get(), "HoverHintOnly")
 					.BrowseButtonToolTip(LOCTEXT("FileButtonToolTipText", "Choose a file from this computer"))
+					.BrowseDirectory(this, &FMediaPlateResourceCustomization::GetMediaBrowseDirectory)
 					.BrowseTitle(LOCTEXT("PropertyEditorTitle", "File picker..."))
 					.FilePath(this, &FMediaPlateResourceCustomization::GetMediaPath)
 					.FileTypeFilter(TEXT("All files (*.*)|*.*"))
@@ -185,11 +187,134 @@ FString FMediaPlateResourceCustomization::GetMediaPath() const
 	return FString();
 }
 
+namespace UE::MediaPlateResourceCustomization::Private
+{
+	/**
+	 * Returns list of possible media base paths in order of priority. 
+	 */
+	const TArray<FString>& GetPossibleBasePaths()
+	{
+		static const TArray<FString> PossibleBasePaths
+		{
+			FPaths::ProjectContentDir(),
+			FPaths::ProjectDir()
+		};
+		return PossibleBasePaths;
+	}
+
+	FString EnsureStartWithDotSlash(const FString& InPath)
+	{
+		// Note: exception for tokens that start with "{".		
+		if (!InPath.StartsWith(TEXT("./")) && !InPath.StartsWith(TEXT("{")))
+		{
+			return FPaths::Combine(TEXT("."), InPath);
+		}
+		return InPath;
+	}
+
+	/** Converts the given path relative to one of the possible base paths. */
+	bool TryConvertAbsoluteToRelative(FString& InAbsolutePath)
+	{
+		FString ConvertedPath = InAbsolutePath;
+
+		for (const FString& BasePath : GetPossibleBasePaths())
+		{
+			const FString FullBasePath = FPaths::ConvertRelativePathToFull(BasePath);
+			if (FPaths::IsUnderDirectory(ConvertedPath, FullBasePath) && FPaths::MakePathRelativeTo(ConvertedPath, *FullBasePath)) 
+			{
+				InAbsolutePath = EnsureStartWithDotSlash(ConvertedPath);
+				return true;
+			}
+		}
+		return false;
+	}
+	
+	/**
+	 * Returns a sanitized path compliant with the path resolution rules
+	 * of ImgMediaSource and FileMediaSource.
+	 * 
+	 * @param InPickedPath Picked path to an existing file.
+	 * @return sanitized path 
+	 */
+	FString SanitizePickedPath(const FString& InPickedPath)
+	{
+		if (InPickedPath.IsEmpty())
+		{
+			return InPickedPath;	
+		}
+
+		FString NormalizedPath = InPickedPath.TrimStartAndEnd().Replace(TEXT("\""), TEXT(""));	
+		FPaths::NormalizeDirectoryName(NormalizedPath);
+
+		if (FPaths::IsRelative(NormalizedPath))
+		{
+			// 1- Try to resolve as relative to BaseDir...
+			FString FullMediaPath = FPaths::ConvertRelativePathToFull(NormalizedPath);
+			if (FPaths::FileExists(FullMediaPath))
+			{
+				// Convert absolute path to relative (if possible), leave absolute if not.
+				TryConvertAbsoluteToRelative(FullMediaPath);
+				return FullMediaPath;
+			}
+
+			// 2- Try to find under which possible base path this path is relative to.
+			for (const FString& BasePath : GetPossibleBasePaths())
+			{
+				const FString FullBasePath = FPaths::ConvertRelativePathToFull(BasePath);
+				FString CombinedMediaPath = FPaths::Combine(FullBasePath, NormalizedPath);
+				
+				if (FPaths::FileExists(CombinedMediaPath) && FPaths::MakePathRelativeTo(CombinedMediaPath, *FullBasePath))
+				{
+					return EnsureStartWithDotSlash(CombinedMediaPath);
+				}
+			}
+
+			// 3- Couldn't find a base, leave as is, but make sure it has a ./
+			return EnsureStartWithDotSlash(NormalizedPath);
+		}
+
+		// Convert absolute path to relative (if possible), leave absolute if not.
+		TryConvertAbsoluteToRelative(NormalizedPath);
+		return NormalizedPath;
+	}
+}
+
+FString FMediaPlateResourceCustomization::GetMediaBrowseDirectory() const
+{
+	const FString MediaPath = GetMediaPath();
+	if (!MediaPath.IsEmpty())
+	{
+		const FString MediaDirectory = FPaths::GetPath(MediaPath);
+		
+		if (FPaths::DirectoryExists(MediaDirectory))
+		{
+			return MediaDirectory;
+		}
+		
+		if (FPaths::IsRelative(MediaDirectory))
+		{
+			for (const FString& BasePath : UE::MediaPlateResourceCustomization::Private::GetPossibleBasePaths())
+			{
+				FString ExpandedDirectory = FPaths::ConvertRelativePathToFull(FPaths::Combine(BasePath, MediaDirectory));
+				if (FPaths::DirectoryExists(ExpandedDirectory))
+				{
+					return ExpandedDirectory;
+				}
+			}
+		}
+	}
+
+	// Fallback to last opened directory.
+	return FEditorDirectories::Get().GetLastDirectory(ELastDirectory::GENERIC_OPEN);
+}
+
 void FMediaPlateResourceCustomization::OnMediaPathPicked(const FString& InPickedPath)
 {
 	if (ExternalMediaPathPropertyHandle)
 	{
-		ExternalMediaPathPropertyHandle->SetValue(InPickedPath);
+		const FString SanitizedMediaPath = UE::MediaPlateResourceCustomization::Private::SanitizePickedPath(InPickedPath);
+		ExternalMediaPathPropertyHandle->SetValue(SanitizedMediaPath);
+		FEditorDirectories::Get().SetLastDirectory(ELastDirectory::GENERIC_OPEN, FPaths::GetPath(InPickedPath));
 	}
 }
 
diff --git a/Engine/Plugins/Media/MediaPlate/Source/MediaPlateEditor/Private/MediaPlateResourceCustomization.h b/Engine/Plugins/Media/MediaPlate/Source/MediaPlateEditor/Private/MediaPlateResourceCustomization.h
--- a/Engine/Plugins/Media/MediaPlate/Source/MediaPlateEditor/Private/MediaPlateResourceCustomization.h
+++ b/Engine/Plugins/Media/MediaPlate/Source/MediaPlateEditor/Private/MediaPlateResourceCustomization.h
@@ -28,6 +28,7 @@ private:
 	TOptional<EMediaPlateResourceType> GetAssetType() const;
 
 	FString GetMediaPath() const;
+	FString GetMediaBrowseDirectory() const;
 
 	void OnAssetTypeChanged(TOptional<EMediaPlateResourceType> InResourceType);
 	void OnMediaPathPicked(const FString& InPickedPath);
diff --git a/Engine/Plugins/Media/MediaPlayerEditor/Source/MediaPlayerEditor/Private/Customizations/FileMediaSourceCustomization.cpp b/Engine/Plugins/Media/MediaPlayerEditor/Source/MediaPlayerEditor/Private/Customizations/FileMediaSourceCustomization.cpp
--- a/Engine/Plugins/Media/MediaPlayerEditor/Source/MediaPlayerEditor/Private/Customizations/FileMediaSourceCustomization.cpp
+++ b/Engine/Plugins/Media/MediaPlayerEditor/Source/MediaPlayerEditor/Private/Customizations/FileMediaSourceCustomization.cpp
@@ -1,20 +1,21 @@
 // Copyright Epic Games, Inc. All Rights Reserved.
 
 #include "Customizations/FileMediaSourceCustomization.h"
+#include "DetailCategoryBuilder.h"
+#include "DetailLayoutBuilder.h"
+#include "DetailWidgetRow.h"
+#include "FileMediaSource.h"
+#include "IDetailPropertyRow.h"
+#include "IMediaModule.h"
 #include "Misc/Paths.h"
+#include "Modules/ModuleManager.h"
+#include "Styling/AppStyle.h"
+#include "Styling/CoreStyle.h"
 #include "Widgets/DeclarativeSyntaxSupport.h"
+#include "Widgets/Images/SImage.h"
+#include "Widgets/Input/SFilePathPicker.h"
 #include "Widgets/SBoxPanel.h"
-#include "Styling/CoreStyle.h"
 #include "Widgets/Text/STextBlock.h"
-#include "Styling/AppStyle.h"
-#include "Modules/ModuleManager.h"
-#include "Widgets/Input/SFilePathPicker.h"
-#include "Widgets/Images/SImage.h"
-#include "DetailWidgetRow.h"
-#include "IDetailPropertyRow.h"
-#include "DetailCategoryBuilder.h"
-#include "DetailLayoutBuilder.h"
-#include "IMediaModule.h"
 
 
 #define LOCTEXT_NAMESPACE "FFileMediaSourceCustomization"
@@ -70,7 +71,7 @@ void FFileMediaSourceCustomization::CustomizeDetails(IDetailLayoutBuilder& Detai
 							.BrowseButtonImage(FAppStyle::GetBrush("PropertyWindow.Button_Ellipsis"))
 							.BrowseButtonStyle(FAppStyle::Get(), "HoverHintOnly")
 							.BrowseButtonToolTip(LOCTEXT("FilePathBrowseButtonToolTip", "Choose a file from this computer"))
-							.BrowseDirectory(FPaths::ProjectContentDir() / TEXT("Movies"))
+							.BrowseDirectory(this, &FFileMediaSourceCustomization::HandleFilePathBrowseDirectory)
 							.FilePath(this, &FFileMediaSourceCustomization::HandleFilePathPickerFilePath)
 							.FileTypeFilter(this, &FFileMediaSourceCustomization::HandleFilePathPickerFileTypeFilter)
 							.OnPathPicked(this, &FFileMediaSourceCustomization::HandleFilePathPickerPathPicked)
@@ -81,9 +82,30 @@ void FFileMediaSourceCustomization::CustomizeDetails(IDetailLayoutBuilder& Detai
 }
 
 
+FString FFileMediaSourceCustomization::GetResolvedFilePath() const
+{
+	TArray<UObject*> OuterObjects;
+	FilePathProperty->GetOuterObjects(OuterObjects);
+	for (UObject* OuterObject : OuterObjects)
+	{
+		if (const UFileMediaSource* Source = Cast<UFileMediaSource>(OuterObject))
+		{
+			// Using the media source to resolve the full path.
+			return Source->GetFullPath();
+		}
+	}
+	return FString();
+}
+
 /* FFileMediaSourceCustomization callbacks
  *****************************************************************************/
 
+FString FFileMediaSourceCustomization::HandleFilePathBrowseDirectory() const
+{
+	const FString MediaPath = GetResolvedFilePath();
+	return MediaPath.IsEmpty() ? (FPaths::ProjectContentDir() / TEXT("Movies")) : FPaths::GetPath(MediaPath);
+}
+
 FString FFileMediaSourceCustomization::HandleFilePathPickerFilePath() const
 {
 	FString FilePath;
@@ -131,49 +153,124 @@ FString FFileMediaSourceCustomization::HandleFilePathPickerFileTypeFilter() cons
 	return Filter;
 }
 
-
-void FFileMediaSourceCustomization::HandleFilePathPickerPathPicked(const FString& PickedPath)
+namespace UE::FileMediaSourceCustomization::Private
 {
-	if (PickedPath.IsEmpty() || PickedPath.StartsWith(TEXT("./")))
+	/**
+	 * Returns list of possible media base paths in order of priority. 
+	 */
+	const TArray<FString>& GetPossibleBasePaths()
 	{
-		FilePathProperty->SetValue(PickedPath);
+		static const TArray<FString> PossibleBasePaths
+		{
+			FPaths::ProjectContentDir(),
+			FPaths::ProjectDir()
+		};
+		return PossibleBasePaths;
 	}
-	else
-	{	
-		FString FullPath = FPaths::ConvertRelativePathToFull(PickedPath);
-		const FString FullGameContentDir = FPaths::ConvertRelativePathToFull(FPaths::ProjectContentDir());
 
-		if (FullPath.StartsWith(FullGameContentDir))
+	FString EnsureStartWithDotSlash(const FString& InPath)
+	{
+		if (!InPath.StartsWith(TEXT("./")))
 		{
-			FPaths::MakePathRelativeTo(FullPath, *FullGameContentDir);
-			FullPath = FString(TEXT("./")) + FullPath;
+			return FPaths::Combine(TEXT("."), InPath);
 		}
+		return InPath;
+	}
 
-		FilePathProperty->SetValue(FullPath);
+	/** Converts the given path relative to one of the possible base paths. */
+	bool TryConvertAbsoluteToRelative(FString& InAbsolutePath)
+	{
+		FString ConvertedPath = InAbsolutePath;
+
+		for (const FString& BasePath : GetPossibleBasePaths())
+		{
+			const FString FullBasePath = FPaths::ConvertRelativePathToFull(BasePath);
+			if (FPaths::IsUnderDirectory(ConvertedPath, FullBasePath) && FPaths::MakePathRelativeTo(ConvertedPath, *FullBasePath)) 
+			{
+				InAbsolutePath = EnsureStartWithDotSlash(ConvertedPath);
+				return true;
+			}
+		}
+		return false;
+	}
+	
+	/**
+	 * Returns a sanitized path compliant with the path resolution rules
+	 * of ImgMediaSource and FileMediaSource.
+	 * 
+	 * @param InPickedPath Picked path to an existing file.
+	 * @return sanitized path 
+	 */
+	FString SanitizePickedPath(const FString& InPickedPath)
+	{
+		if (InPickedPath.IsEmpty())
+		{
+			return InPickedPath;	
+		}
+
+		FString NormalizedPath = InPickedPath.TrimStartAndEnd().Replace(TEXT("\""), TEXT(""));	
+		FPaths::NormalizeDirectoryName(NormalizedPath);
+
+		if (FPaths::IsRelative(NormalizedPath))
+		{
+			// 1- Try to resolve as relative to BaseDir...
+			FString FullMediaPath = FPaths::ConvertRelativePathToFull(NormalizedPath);
+			if (FPaths::FileExists(FullMediaPath))
+			{
+				// Convert absolute path to relative (if possible), leave absolute if not.
+				TryConvertAbsoluteToRelative(FullMediaPath);
+				return FullMediaPath;
+			}
+
+			// 2- Try to find under which possible base path this path is relative to.
+			for (const FString& BasePath : GetPossibleBasePaths())
+			{
+				const FString FullBasePath = FPaths::ConvertRelativePathToFull(BasePath);
+				FString CombinedMediaPath = FPaths::Combine(FullBasePath, NormalizedPath);
+				
+				if (FPaths::FileExists(CombinedMediaPath) && FPaths::MakePathRelativeTo(CombinedMediaPath, *FullBasePath))
+				{
+					return EnsureStartWithDotSlash(CombinedMediaPath);
+				}
+			}
+
+			// 3- Couldn't find a base, leave as is, but make sure it has a ./
+			return EnsureStartWithDotSlash(NormalizedPath);
+		}
+
+		// Convert absolute path to relative (if possible), leave absolute if not.
+		TryConvertAbsoluteToRelative(NormalizedPath);
+		return NormalizedPath;
 	}
 }
 
+void FFileMediaSourceCustomization::HandleFilePathPickerPathPicked(const FString& PickedPath)
+{
+	if (FilePathProperty)
+	{
+		const FString SanitizedMediaPath = UE::FileMediaSourceCustomization::Private::SanitizePickedPath(PickedPath);
+		FilePathProperty->SetValue(SanitizedMediaPath);
+	}
+}
 
 EVisibility FFileMediaSourceCustomization::HandleFilePathWarningIconVisibility() const
 {
-	FString FilePath;
+	const FString FilePath = GetResolvedFilePath();
 
-	if ((FilePathProperty->GetValue(FilePath) != FPropertyAccess::Success) || FilePath.IsEmpty() || FilePath.Contains(TEXT("://")))
+	if (FilePath.IsEmpty())
 	{
 		return EVisibility::Hidden;
 	}
 
 	const FString FullMoviesPath = FPaths::ConvertRelativePathToFull(FPaths::ProjectContentDir() / TEXT("Movies"));
-	const FString FullPath = FPaths::ConvertRelativePathToFull(FPaths::IsRelative(FilePath) ? FPaths::ProjectContentDir() / FilePath : FilePath);
-
-	if (FullPath.StartsWith(FullMoviesPath))
+	if (FPaths::IsUnderDirectory(FilePath, FullMoviesPath))
 	{
-		if (FPaths::FileExists(FullPath))
+		if (FPaths::FileExists(FilePath))
 		{
 			return EVisibility::Hidden;
 		}
 
-		// file doesn't exist
+		// Path doesn't exist
 		return EVisibility::Visible;
 	}
 
diff --git a/Engine/Plugins/Media/MediaPlayerEditor/Source/MediaPlayerEditor/Private/Customizations/FileMediaSourceCustomization.h b/Engine/Plugins/Media/MediaPlayerEditor/Source/MediaPlayerEditor/Private/Customizations/FileMediaSourceCustomization.h
--- a/Engine/Plugins/Media/MediaPlayerEditor/Source/MediaPlayerEditor/Private/Customizations/FileMediaSourceCustomization.h
+++ b/Engine/Plugins/Media/MediaPlayerEditor/Source/MediaPlayerEditor/Private/Customizations/FileMediaSourceCustomization.h
@@ -34,7 +34,12 @@ public:
 	}
 
 private:
+	/** Returns the fully resolved path of the selected media file. */
+	FString GetResolvedFilePath() const;
 
+	/** Callback for getting the browse directory. */
+	FString HandleFilePathBrowseDirectory() const;
+	
 	/** Callback for getting the selected path in the URL picker widget. */
 	FString HandleFilePathPickerFilePath() const;
 
diff --git a/Engine/Source/Runtime/MediaAssets/Private/Assets/FileMediaSource.cpp b/Engine/Source/Runtime/MediaAssets/Private/Assets/FileMediaSource.cpp
--- a/Engine/Source/Runtime/MediaAssets/Private/Assets/FileMediaSource.cpp
+++ b/Engine/Source/Runtime/MediaAssets/Private/Assets/FileMediaSource.cpp
@@ -80,13 +80,23 @@ void UFileMediaSource::ResolveFullPath() const
     
 	if (ResolvedFullPath.StartsWith(TEXT("./")))
 	{
-        ResolvedFullPath = FPaths::Combine(FPaths::ProjectContentDir(), FilePath.RightChop(2));
-        FString FinalFullPath = FPaths::ConvertRelativePathToFull(ResolvedFullPath);
-        if (FPaths::FileExists(FinalFullPath))
-        {
-            ResolvedFullPath = FinalFullPath;
-            return;
-        }
+		// Support Content or Project relative (in that order)
+		const TArray<FString> PossibleBasePaths =
+		{
+			FPaths::ProjectContentDir(),
+			FPaths::ProjectDir()
+		};
+
+		for (const FString& BasePath : PossibleBasePaths)
+		{
+			ResolvedFullPath = FPaths::Combine(BasePath, FilePath.RightChop(2));
+			FString FinalFullPath = FPaths::ConvertRelativePathToFull(ResolvedFullPath);
+			if (FPaths::FileExists(FinalFullPath))
+			{
+				ResolvedFullPath = FinalFullPath;
+				return;
+			}
+		}
     }
 
 	const TArray<FString>& RootDirectories = FPlatformMisc::GetAdditionalRootDirectories();
