[ImgMedia] Fix scrubbing frame accuracy

FrameToTime and TimeToFrame must be reversible to avoid issues, otherwise it is very complicated to get the frame selection to work.
Unit tests have been added to ensure this because it is the foundation of frame accuracy.

diff --git a/Engine/Plugins/Media/ImgMedia/Source/ImgMedia/Private/Loader/ImgMediaLoader.cpp b/Engine/Plugins/Media/ImgMedia/Source/ImgMedia/Private/Loader/ImgMediaLoader.cpp
--- a/Engine/Plugins/Media/ImgMedia/Source/ImgMedia/Private/Loader/ImgMediaLoader.cpp
+++ b/Engine/Plugins/Media/ImgMedia/Source/ImgMedia/Private/Loader/ImgMediaLoader.cpp
@@ -19,6 +19,7 @@
 #include "Readers/GenericImgMediaReader.h"
 #include "IImageWrapperModule.h"
 #include "IImgMediaModule.h"
+#include "ImgMediaLoaderUtils.h"
 #include "Readers/IImgMediaReader.h"
 #include "Readers/DdsImgMediaReader.h"
 #include "ImgMediaLoaderWork.h"
@@ -346,7 +347,7 @@ void FImgMediaLoader::Flush()
 void FImgMediaLoader::Seek(const FMediaTimeStamp& SeekTarget, bool bReverse)
 {
 	ResetFetchLogic();
-	QueuedSampleFetch.LastTimeStamp.SetTimeAndIndexValue(FrameNumberToTime(TimeToFrameNumber(SeekTarget.Time)), SeekTarget.GetIndexValue());
+	QueuedSampleFetch.LastTimeStamp.SetTimeAndIndexValue(SeekTarget.Time, SeekTarget.GetIndexValue());
 	QueuedSampleFetch.LastDuration = FTimespan::Zero();
 }
 
@@ -377,7 +378,6 @@ bool FImgMediaLoader::DiscardVideoSamples(const TRange<FMediaTimeStamp>& TimeRan
 	if (bOk)
 	{
 		ResetFetchLogic();
-		QueuedSampleFetch.LastTimeStamp.Time = FrameNumberToTime(TimeToFrameNumber(QueuedSampleFetch.LastTimeStamp.Time));
 		QueuedSampleFetch.LastDuration = FTimespan::Zero();
 	}
 
@@ -628,8 +628,6 @@ bool FImgMediaLoader::PeekVideoSampleTime(FMediaTimeStamp &TimeStamp, bool bIsLo
 			}
 		}
 
-		TimeStamp.Time = FrameNumberToTime(TimeToFrameNumber(TimeStamp.Time));
-
 		// If waiting for a seek and there is no other playback blocking,
 		// properly indicate if the sample has arrived to allow for seek throttling.
 		if (QueuedSampleFetch.LastDuration.IsZero() && !bIsPlaybackBlocking)
@@ -761,8 +759,6 @@ void FImgMediaLoader::Initialize(const FString& SequencePath, const FFrameRate&
 
 bool FImgMediaLoader::RequestFrame(FTimespan Time, float PlayRate, bool Loop)
 {
-	// Match the frame rounding done in FImgMediaLoader::Seek to ensure requested frame is the same as in PeekVideoSampleTime.
-	Time = FrameNumberToTime(TimeToFrameNumber(Time));
 	return RequestFrame(TimeToFrameNumber(Time), PlayRate, Loop);
 }
 
@@ -1166,22 +1162,18 @@ uint32 FImgMediaLoader::TimeToFrameNumber(FTimespan Time) const
 	* cases covered by the older code.
 	*/
 
-	// note: do NOT assume 100% reversability between TimeToFrameNumber() and FrameNumberToTime()! The FTimespan "only" has 100ns resolution and does some rounding when converting from seconds -> this migth shift things around!
-	return static_cast<uint32>(FMath::FloorToInt32(Time.GetTotalSeconds() * (static_cast<double>(SequenceFrameRate.Numerator) / SequenceFrameRate.Denominator)));
+	return UE::ImgMediaLoaderUtils::TimeToFrameNumber(Time, SequenceFrameRate);
 }
 
-
 FTimespan FImgMediaLoader::FrameNumberToTime(uint32 FrameNumber) const
 {
-	// note: do NOT assume 100% reversability between TimeToFrameNumber() and FrameNumberToTime()! The FTimespan "only" has 100ns resolution and does some rounding when converting from seconds -> this migth shift things around!
-	return FTimespan::FromSeconds(FrameNumber * (static_cast<double>(SequenceFrameRate.Denominator) / SequenceFrameRate.Numerator));
+	// Note: this version is reversible with TimeToFrameNumber.
+	return UE::ImgMediaLoaderUtils::GetFrameStartTime(FrameNumber, SequenceFrameRate);
 }
 
-
 int64 FImgMediaLoader::TimeToFrameNumberUnbound(FTimespan Time) const
 {
-	// note: do NOT assume 100% reversability between TimeToFrameNumber() and FrameNumberToTime()! The FTimespan "only" has 100ns resolution and does some rounding when converting from seconds -> this migth shift things around!
-	return static_cast<int64>(FMath::FloorToInt(Time.GetTotalSeconds() * (static_cast<double>(SequenceFrameRate.Numerator) / SequenceFrameRate.Denominator)));
+	return UE::ImgMediaLoaderUtils::TimeToFrameNumberUnbound(Time, SequenceFrameRate);
 }
 
 
diff --git a/Engine/Plugins/Media/ImgMedia/Source/ImgMedia/Private/Loader/ImgMediaLoaderUtils.h b/Engine/Plugins/Media/ImgMedia/Source/ImgMedia/Private/Loader/ImgMediaLoaderUtils.h
new file mode 100644
--- /dev/null
+++ b/Engine/Plugins/Media/ImgMedia/Source/ImgMedia/Private/Loader/ImgMediaLoaderUtils.h
@@ -0,0 +1,58 @@
+// Copyright Epic Games, Inc. All Rights Reserved.
+
+#pragma once
+
+#include "Misc/FrameRate.h"
+#include "Misc/Timespan.h"
+
+/**
+ * Frame conversion math utilities.
+ * Extracted from FImgMediaLoader for unit tests.
+ */
+namespace UE::ImgMediaLoaderUtils
+{
+	/**
+	 * Converts the given time to a frame number at the given framerate.
+	 * @param InTime Time (timespan in ticks)
+	 * @param InFrameRate Frame rate in frames per second
+	 * @return Frame number
+	 */
+	inline uint32 TimeToFrameNumber(const FTimespan& InTime, const FFrameRate& InFrameRate)
+	{
+		return static_cast<uint32>(FMath::FloorToInt32(InTime.GetTotalSeconds() * InFrameRate.AsDecimal()));
+	}
+
+	/**
+	 * Converts the given time to a frame number at the given framerate.
+	 * @param InTime Time (timespan in ticks)
+	 * @param InFrameRate Frame rate in frames per second
+	 * @return Frame number (can be very large)
+	 */
+	inline int64 TimeToFrameNumberUnbound(const FTimespan& InTime, const FFrameRate& InFrameRate)
+	{
+		return static_cast<int64>(FMath::FloorToInt(InTime.GetTotalSeconds() * InFrameRate.AsDecimal()));
+	}
+	
+	/**
+	 * Returns the lower bound, start time, of the given frame number for the given frame rate.
+	 * 
+	 * @param InFrameNumber Frame number
+	 * @param InFrameRate Frame rate in frames per second
+	 * @return A timespan in tick, where the tick number is the first tick that will return back the same frame number from TimeToFrameNumber.
+	 */
+	inline FTimespan GetFrameStartTime(uint32 InFrameNumber, const FFrameRate& InFrameRate)
+	{
+		const double Ticks = InFrameNumber * InFrameRate.AsInterval() * ETimespan::TicksPerSecond;
+		
+		const FTimespan FloorTicks(FMath::FloorToInt64(Ticks));
+
+		// Test the case where floor of tick will be reversible, in that case, it is the first tick of that frame.
+		if (TimeToFrameNumber(FloorTicks, InFrameRate) == InFrameNumber)
+		{
+			return FloorTicks;
+		}
+
+		// Ceil of Ticks is always reversible with TimeToFrameNumber, but is not always the first tick of the frame.
+		return FTimespan(FMath::CeilToInt64(Ticks));
+	}
+}
\ No newline at end of file
diff --git a/Engine/Plugins/Media/ImgMedia/Source/ImgMedia/Private/Tests/ImgMediaTests.cpp b/Engine/Plugins/Media/ImgMedia/Source/ImgMedia/Private/Tests/ImgMediaTests.cpp
--- a/Engine/Plugins/Media/ImgMedia/Source/ImgMedia/Private/Tests/ImgMediaTests.cpp
+++ b/Engine/Plugins/Media/ImgMedia/Source/ImgMedia/Private/Tests/ImgMediaTests.cpp
@@ -1,14 +1,14 @@
 // Copyright Epic Games, Inc. All Rights Reserved.
 
+#include "Assets/ImgMediaMipMapInfo.h"
 #include "Containers/UnrealString.h"
 #include "CoreTypes.h"
 #include "HAL/FileManager.h"
+#include "ImgMediaSource.h"
+#include "Loader/ImgMediaLoaderUtils.h"
 #include "Logging/LogMacros.h"
 #include "Misc/AutomationTest.h"
 
-#include "Assets/ImgMediaMipMapInfo.h"
-#include "ImgMediaSource.h"
-
 #if WITH_DEV_AUTOMATION_TESTS
 
 IMPLEMENT_SIMPLE_AUTOMATION_TEST(FImgMediaTests, "System.Plugins.ImgMedia.TileSelection", EAutomationTestFlags::EditorContext | EAutomationTestFlags::EngineFilter)
@@ -91,6 +91,50 @@ bool FImgMediaTests::RunTest(const FString& Parameters)
 	return !HasAnyErrors();
 }
 
+IMPLEMENT_SIMPLE_AUTOMATION_TEST(FImgMediaLoaderUtilsTests, "System.Plugins.ImgMedia.LoaderUtils", EAutomationTestFlags::EditorContext | EAutomationTestFlags::EngineFilter)
+bool FImgMediaLoaderUtilsTests::RunTest(const FString& Parameters)
+{
+	TArray<FFrameRate> FrameRates =
+	{
+	{60000, 1001},	// 59.95
+	{30000, 1001},	// 29.97
+	{24000, 1001},	// 23.976
+	{60, 1},			// 60
+	{30, 1},			// 30
+	{50, 1},			// 50
+	{25, 1},			// 25
+	{24, 1}			// 24
+	};
+
+
+	// Test to ensure reversibility of Frame to Time conversion functions
+	// This is necessary for frame accuracy and consistency with the sequencer.
+	for (FFrameRate FrameRate : FrameRates)
+	{
+		for (uint32 FrameNumber = 0; FrameNumber < 100; FrameNumber++)
+		{
+			{
+				FTimespan StartTime = UE::ImgMediaLoaderUtils::GetFrameStartTime(FrameNumber, FrameRate);
+				uint32 ConvertedFrameNumber = UE::ImgMediaLoaderUtils::TimeToFrameNumber(StartTime, FrameRate);			
+				AddErrorIfFalse(ConvertedFrameNumber == FrameNumber, FString::Printf(TEXT("FImgMediaLoaderUtilsTests: Frame Start Time %lld conversion at %f not reversible: Frame %i != %i"), StartTime.GetTicks(), FrameRate.AsDecimal(), FrameNumber, ConvertedFrameNumber));
+				int64 ConvertedFrameNumber64 = UE::ImgMediaLoaderUtils::TimeToFrameNumberUnbound(StartTime, FrameRate);
+				AddErrorIfFalse(ConvertedFrameNumber64 == static_cast<int64>(FrameNumber), FString::Printf(TEXT("FImgMediaLoaderUtilsTests: Unbound Frame Start Time %lld conversion at %f not reversible: Frame %i != %lld"), StartTime.GetTicks(), FrameRate.AsDecimal(), FrameNumber, ConvertedFrameNumber64));
+			}
+			{
+				// Testing that GetFrameStartTime is indeed the first tick of the frame.
+				// One tick before the start of next frame should be the previous frame.
+				FTimespan EndTime = UE::ImgMediaLoaderUtils::GetFrameStartTime(FrameNumber + 1, FrameRate) - FTimespan(1);
+				uint32 ConvertedFrameNumber = UE::ImgMediaLoaderUtils::TimeToFrameNumber(EndTime, FrameRate);
+				AddErrorIfFalse(ConvertedFrameNumber == FrameNumber, FString::Printf(TEXT("FImgMediaLoaderUtilsTests: Frame End Time %lld conversion at %f not reversible: Frame %i != %i"), EndTime.GetTicks(), FrameRate.AsDecimal(), FrameNumber, ConvertedFrameNumber));
+				int64 ConvertedFrameNumber64 = UE::ImgMediaLoaderUtils::TimeToFrameNumberUnbound(EndTime, FrameRate);
+				AddErrorIfFalse(ConvertedFrameNumber64 == static_cast<int64>(FrameNumber), FString::Printf(TEXT("FImgMediaLoaderUtilsTests: Unbound Frame End Time %lld conversion at %f not reversible: Frame %i != %lld"), EndTime.GetTicks(), FrameRate.AsDecimal(), FrameNumber, ConvertedFrameNumber64));
+			}
+		}
+	}
+	
+	return !HasAnyErrors();
+}
+
 IMPLEMENT_SIMPLE_AUTOMATION_TEST(FImgMediaSanitizePathTests, "System.Plugins.ImgMedia.SanitizePath", EAutomationTestFlags::EditorContext | EAutomationTestFlags::EngineFilter)
 bool FImgMediaSanitizePathTests::RunTest(const FString& Parameters)
 {
