ImgMedia: Disable (exr gpu reader) pre-allocate pool with sequencer playback.

With sequencer, the warmup pre-allocations had negative consequences outweighting any potential benefits:
* Prevented asynchronous work in pre-roll section (now fixed) due to spike in render thread.
* In cases where only a few mips or tiles are visible, the allocations were overly greedy.
* Increased risk of blocking playback disabling if warmup allocations exceed timeout.

Furthermore if read times exceed frame times, preallocations still won't prevent blocks in the long run.

diff --git a/Engine/Plugins/Media/ImgMedia/Source/ImgMedia/Internal/Loader/ImgMediaLoader.h b/Engine/Plugins/Media/ImgMedia/Source/ImgMedia/Internal/Loader/ImgMediaLoader.h
--- a/Engine/Plugins/Media/ImgMedia/Source/ImgMedia/Internal/Loader/ImgMediaLoader.h
+++ b/Engine/Plugins/Media/ImgMedia/Source/ImgMedia/Internal/Loader/ImgMediaLoader.h
@@ -334,9 +334,10 @@ public:
 	 * @param SequencePath Path to the image sequence.
 	 * @param FrameRateOverride The frame rate to use (0/0 = do not override).
 	 * @param Loop Whether the cache should loop around.
+	 * @param bIsSequencerEnvironment Whether image sequence playback comes from sequencer, to decide on warmup behavior.
 	 * @see IsInitialized
 	 */
-	void Initialize(const FString& SequencePath, const FFrameRate& FrameRateOverride, bool Loop);
+	void Initialize(const FString& SequencePath, const FFrameRate& FrameRateOverride, bool Loop, bool bIsSequencerEnvironment);
 
 	/**
 	 * Whether this loader has been initialized yet.
@@ -440,8 +441,9 @@ protected:
 	 *
 	 * @param InFirstFrameInfo Sequence information based on its first frame.
 	 * @param Loop Whether the cache should loop around.
+	 * @param bIsSequencerEnvironment Whether image sequence playback comes from sequencer.
 	 */
-	void WarmupSequence(const FImgMediaFrameInfo& InFirstFrameInfo, bool Loop);
+	void WarmupSequence(const FImgMediaFrameInfo& InFirstFrameInfo, bool Loop, bool bIsSequencerEnvironment);
 
 	/**
 	 * Finds all the files in a directory and gets their path.
diff --git a/Engine/Plugins/Media/ImgMedia/Source/ImgMedia/Private/Loader/ImgMediaLoader.cpp b/Engine/Plugins/Media/ImgMedia/Source/ImgMedia/Private/Loader/ImgMediaLoader.cpp
--- a/Engine/Plugins/Media/ImgMedia/Source/ImgMedia/Private/Loader/ImgMediaLoader.cpp
+++ b/Engine/Plugins/Media/ImgMedia/Source/ImgMedia/Private/Loader/ImgMediaLoader.cpp
@@ -733,7 +733,7 @@ IQueuedWork* FImgMediaLoader::GetWork()
 
 
 void FImgMediaLoader::Initialize(const FString& SequencePath, const FFrameRate& FrameRateOverride,
-	bool Loop)
+	bool Loop, bool bIsSequencerEnvironment)
 {
 	UE_LOG(LogImgMedia, Verbose, TEXT("Loader %p: Initializing with %s (FrameRateOverride = %s, Loop = %i)"),
 		this,
@@ -748,7 +748,7 @@ void FImgMediaLoader::Initialize(const FString& SequencePath, const FFrameRate&
 
 	if (LoadSequence(SequencePath, FrameRateOverride, Loop, FirstFrameInfo))
 	{
-		WarmupSequence(FirstFrameInfo, Loop);
+		WarmupSequence(FirstFrameInfo, Loop, bIsSequencerEnvironment);
 	}
 
 	FPlatformMisc::MemoryBarrier();
@@ -989,12 +989,21 @@ bool FImgMediaLoader::LoadSequence(const FString& SequencePath, const FFrameRate
 	return true;
 }
 
-void FImgMediaLoader::WarmupSequence(const FImgMediaFrameInfo& InFirstFrameInfo, bool Loop)
+void FImgMediaLoader::WarmupSequence(const FImgMediaFrameInfo& InFirstFrameInfo, bool Loop, bool bIsSequencerEnvironment)
 {
-	// Giving our reader a chance to handle RAM allocation.
-	// Not all readers use this, only those that need to handle large files
-	// or need to be as efficient as possible.
-	Reader->PreAllocateMemoryPool(NumFramesToLoad, InFirstFrameInfo);
+	/**
+	* With sequencer, the warmup pre-allocations had negative consequences outweighting any potential benefits:
+	* - Prevented asynchronous work in pre-roll section due to spike in render thread.
+	* - In cases where only a few mips or tiles are visible, the allocations were overly greedy.
+	* - Increased risk of blocking playback disabling if warmup allocations exceed timeout.
+	*/
+	if (!bIsSequencerEnvironment)
+	{
+		// Giving our reader a chance to handle RAM allocation.
+		// Not all readers use this, only those that need to handle large files
+		// or need to be as efficient as possible.
+		Reader->PreAllocateMemoryPool(NumFramesToLoad, InFirstFrameInfo);
+	}
 
 	FScopeLock Lock(&CriticalSection);
 
diff --git a/Engine/Plugins/Media/ImgMedia/Source/ImgMedia/Private/Player/ImgMediaPlayer.cpp b/Engine/Plugins/Media/ImgMedia/Source/ImgMedia/Private/Player/ImgMediaPlayer.cpp
--- a/Engine/Plugins/Media/ImgMedia/Source/ImgMedia/Private/Player/ImgMediaPlayer.cpp
+++ b/Engine/Plugins/Media/ImgMedia/Source/ImgMedia/Private/Player/ImgMediaPlayer.cpp
@@ -241,6 +241,11 @@ bool FImgMediaPlayer::Open(const FString& Url, const IMediaOptions* Options, con
 		const FVariant* Var = InPlayerOptions ? InPlayerOptions->InternalCustomOptions.Find(InOptionName) : nullptr;
 		return Var && Var->GetType() == EVariantTypes::String ? Var->GetValue<FString>() : FString();
 	};
+	auto GetNameOption = [](const FMediaPlayerOptions* InPlayerOptions, const FName& InOptionName) -> FName
+		{
+			const FVariant* Var = InPlayerOptions ? InPlayerOptions->InternalCustomOptions.Find(InOptionName) : nullptr;
+			return Var && Var->GetType() == EVariantTypes::Name ? Var->GetValue<FName>() : FName();
+		};
 
 	// get frame rate override, if any
 	FFrameRate FrameRateOverride(0, 0);
@@ -248,6 +253,7 @@ bool FImgMediaPlayer::Open(const FString& Url, const IMediaOptions* Options, con
 	TSharedPtr<FNativeMediaSourceColorSettings, ESPMode::ThreadSafe> SourceColorSettings;
 	bool bFillGapsInSequence = true;
 	bool bIsSmartCacheEnabled = GetBoolOption(PlayerOptions, MediaPlayerOptionValues::ImgMediaSmartCacheEnabled(), false);
+	bool bIsSequencerEnvironment = GetNameOption(PlayerOptions, MediaPlayerOptionValues::Environment()).IsEqual(MediaPlayerOptionValues::Environment_Sequencer());
 	double SmartCacheTimeToLookAhead = GetDoubleOption(PlayerOptions, MediaPlayerOptionValues::ImgMediaSmartCacheTimeToLookAhead(), 0.0);
 
 	UMediaTexture* ViewTexture = nullptr;
@@ -316,7 +322,7 @@ bool FImgMediaPlayer::Open(const FString& Url, const IMediaOptions* Options, con
 
 	const FString SequencePath = Url.RightChop(6);
 
-	Async(EAsyncExecution::ThreadPool, [FrameRateOverride, LoaderPtr = TWeakPtr<FImgMediaLoader, ESPMode::ThreadSafe>(Loader), Proxy, SequencePath, Loop = ShouldLoop]()
+	Async(EAsyncExecution::ThreadPool, [FrameRateOverride, LoaderPtr = TWeakPtr<FImgMediaLoader, ESPMode::ThreadSafe>(Loader), Proxy, SequencePath, Loop = ShouldLoop, bIsSequencerEnvironment]()
 	{
 		TSharedPtr<FImgMediaLoader, ESPMode::ThreadSafe> PinnedLoader = LoaderPtr.Pin();
 
@@ -329,7 +335,7 @@ bool FImgMediaPlayer::Open(const FString& Url, const IMediaOptions* Options, con
 				ProxyPath = SequencePath; // fall back to root folder
 			}
 
-			PinnedLoader->Initialize(ProxyPath, FrameRateOverride, Loop);
+			PinnedLoader->Initialize(ProxyPath, FrameRateOverride, Loop, bIsSequencerEnvironment);
 		}
 	});
 
