MediaProfile: Adding missing API exports

diff --git a/Engine/Plugins/Media/ImgMedia/Source/ImgMedia/Public/ImgMediaSource.h b/Engine/Plugins/Media/ImgMedia/Source/ImgMedia/Public/ImgMediaSource.h
--- a/Engine/Plugins/Media/ImgMedia/Source/ImgMedia/Public/ImgMediaSource.h
+++ b/Engine/Plugins/Media/ImgMedia/Source/ImgMedia/Public/ImgMediaSource.h
@@ -167,7 +167,7 @@ public:
 	UE_API virtual bool Validate() const override;
 
 #if WITH_EDITOR
-	virtual void GetDetailsPanelInfoElements(TArray<FInfoElement>& OutInfoElements) const override;
+	UE_API virtual void GetDetailsPanelInfoElements(TArray<FInfoElement>& OutInfoElements) const override;
 	
 	UE_API virtual void PostEditChangeProperty(FPropertyChangedEvent& PropertyChangedEvent) override;
 #endif // WITH_EDITOR
