[ImgMedia] Fix FrameToTime and TimeToFrame when compiled with Clang.

Perform the division in integer to avoid floating point precision issues.

diff --git a/Engine/Plugins/Media/ImgMedia/Source/ImgMedia/Private/Loader/ImgMediaLoaderUtils.h b/Engine/Plugins/Media/ImgMedia/Source/ImgMedia/Private/Loader/ImgMediaLoaderUtils.h
--- a/Engine/Plugins/Media/ImgMedia/Source/ImgMedia/Private/Loader/ImgMediaLoaderUtils.h
+++ b/Engine/Plugins/Media/ImgMedia/Source/ImgMedia/Private/Loader/ImgMediaLoaderUtils.h
@@ -19,7 +19,7 @@ namespace UE::ImgMediaLoaderUtils
 	 */
 	inline uint32 TimeToFrameNumber(const FTimespan& InTime, const FFrameRate& InFrameRate)
 	{
-		return static_cast<uint32>(FMath::FloorToInt32(InTime.GetTotalSeconds() * InFrameRate.AsDecimal()));
+		return static_cast<uint32>(FMath::DivideAndRoundDown(InTime.GetTicks() * InFrameRate.Numerator, InFrameRate.Denominator * ETimespan::TicksPerSecond));
 	}
 
 	/**
@@ -30,9 +30,9 @@ namespace UE::ImgMediaLoaderUtils
 	 */
 	inline int64 TimeToFrameNumberUnbound(const FTimespan& InTime, const FFrameRate& InFrameRate)
 	{
-		return static_cast<int64>(FMath::FloorToInt(InTime.GetTotalSeconds() * InFrameRate.AsDecimal()));
+		return FMath::DivideAndRoundDown(InTime.GetTicks() * InFrameRate.Numerator, InFrameRate.Denominator * ETimespan::TicksPerSecond);
 	}
-	
+
 	/**
 	 * Returns the lower bound, start time, of the given frame number for the given frame rate.
 	 * 
@@ -42,9 +42,8 @@ namespace UE::ImgMediaLoaderUtils
 	 */
 	inline FTimespan GetFrameStartTime(uint32 InFrameNumber, const FFrameRate& InFrameRate)
 	{
-		const double Ticks = InFrameNumber * InFrameRate.AsInterval() * ETimespan::TicksPerSecond;
-		
-		const FTimespan FloorTicks(FMath::FloorToInt64(Ticks));
+		const int64 DenominatorTicks = InFrameRate.Denominator * ETimespan::TicksPerSecond;
+		const FTimespan FloorTicks(FMath::DivideAndRoundDown(int64(InFrameNumber) * DenominatorTicks, int64(InFrameRate.Numerator)));
 
 		// Test the case where floor of tick will be reversible, in that case, it is the first tick of that frame.
 		if (TimeToFrameNumber(FloorTicks, InFrameRate) == InFrameNumber)
@@ -53,6 +52,6 @@ namespace UE::ImgMediaLoaderUtils
 		}
 
 		// Ceil of Ticks is always reversible with TimeToFrameNumber, but is not always the first tick of the frame.
-		return FTimespan(FMath::CeilToInt64(Ticks));
+		return FTimespan(FMath::DivideAndRoundUp(int64(InFrameNumber) * DenominatorTicks, int64(InFrameRate.Numerator)));
 	}
 }
\ No newline at end of file
