ImgMedia: Setting frame duration as per the sequence frame rate instead of the value from the global plugin settings. If the frame duration is not correct the calculation in the player facade on whether a frame is still valid for display will be incorrect. Added a comment as a reminder that setting the frame duration like that is probably not correct and also provided default ctor values for trivial members in the info structure for safeties sake.

#rnx

diff --git a/Engine/Plugins/Media/ImgMedia/Source/ImgMedia/Private/Loader/ImgMediaLoader.cpp b/Engine/Plugins/Media/ImgMedia/Source/ImgMedia/Private/Loader/ImgMediaLoader.cpp
--- a/Engine/Plugins/Media/ImgMedia/Source/ImgMedia/Private/Loader/ImgMediaLoader.cpp
+++ b/Engine/Plugins/Media/ImgMedia/Source/ImgMedia/Private/Loader/ImgMediaLoader.cpp
@@ -552,7 +552,7 @@ IMediaSamples::EFetchBestSampleResult FImgMediaLoader::FetchBestVideoSampleForTi
 			QueuedSampleFetch.LastTimeStamp.SetTime(SampleTime);
 			QueuedSampleFetch.LastTimeStamp.SetSequenceIndex(TimeRange.GetLowerBoundValue().GetSequenceIndex());
 			QueuedSampleFetch.LastTimeStamp.SetLoopIndex(QueuedSampleFetch.LoopIndex);
-			QueuedSampleFetch.LastDuration = FTimespan::FromSeconds(Frame->Get()->GetFrameRate().AsInterval());
+			QueuedSampleFetch.LastDuration = FrameLength;
 
 			// We are clear to return it as new result... Make a sample & initialize it...
 			auto Sample = MakeShared<FImgMediaTextureSample, ESPMode::ThreadSafe>();
@@ -902,6 +902,7 @@ bool FImgMediaLoader::LoadSequence(const FString& SequencePath, const FFrameRate
 
 	if (FrameRateOverride.IsValid() && (FrameRateOverride.Numerator > 0))
 	{
+		FirstFrameInfo.FrameRate = FrameRateOverride;
 		SequenceFrameRate = FrameRateOverride;
 	}
 	else
diff --git a/Engine/Plugins/Media/ImgMedia/Source/ImgMedia/Private/Readers/IImgMediaReader.h b/Engine/Plugins/Media/ImgMedia/Source/ImgMedia/Private/Readers/IImgMediaReader.h
--- a/Engine/Plugins/Media/ImgMedia/Source/ImgMedia/Private/Readers/IImgMediaReader.h
+++ b/Engine/Plugins/Media/ImgMedia/Source/ImgMedia/Private/Readers/IImgMediaReader.h
@@ -26,19 +26,26 @@ struct FImgMediaFrameInfo
 	FString FormatName;
 
 	/** Frame rate. */
+	/*
+		Note that the frame rate here is not necessarily correct if the sequence uses a frame rate override.
+		The override value is not always accessible from where this member gets set, and it might also not be
+		all that useful since it is conceivable that the images end up in the image cache and are either
+		referenced by different media sources with different frame rate overrides, or the frame rate override
+		is changed after the images have been added to the cache.
+	*/
 	FFrameRate FrameRate;
 
 	/** Whether the frame is in sRGB color space. */
-	bool Srgb;
+	bool Srgb = false;
 
 	/** Uncompressed size (in bytes). All mip levels included. */
-	SIZE_T UncompressedSize;
+	SIZE_T UncompressedSize = 0;
 
 	/** Number of channels (RGB - 3 or RGBA - 4). */
-	SIZE_T NumChannels;
+	SIZE_T NumChannels = 0;
 
 	/** Does this frame have tiles?. */
-	bool bHasTiles;
+	bool bHasTiles = false;
 
 	/** Number of tiles in X and Y direction. */
 	FIntPoint NumTiles;
@@ -47,10 +54,10 @@ struct FImgMediaFrameInfo
 	FIntPoint TileDimensions;
 
 	/** Tile border size in texels. This is required for more elaborate tile texel sampling. */
-	int32 TileBorder;
+	int32 TileBorder = 0;
 
 	/** Number of Mip Levels. */
-	int32 NumMipLevels;
+	int32 NumMipLevels = 0;
 };
 
 
@@ -63,13 +70,13 @@ struct FImgMediaMipTiles final
 	FImgMediaMipTiles() = default;
 	/** Default destructor. */
 	~FImgMediaMipTiles() = default;
-	
+
 	/** Move constructor. */
 	FImgMediaMipTiles(FImgMediaMipTiles&& Other)
 	{
 		InternalMove(MoveTemp(Other));
 	}
-	
+
 	/** Copy constructor. */
 	FImgMediaMipTiles(const FImgMediaMipTiles& Other)
 	{
@@ -92,7 +99,7 @@ struct FImgMediaMipTiles final
 
 	/**
 	 * Contains the specified mip level.
-	 * 
+	 *
 	 * @param InMipLevel mip level
 	 * @return true when mip level is present.
 	*/
@@ -141,7 +148,7 @@ struct FImgMediaMipTiles final
 
 	/**
 	 * Check if the existing tiles contain all of the requested ones, at the specified mip level.
-	 * 
+	 *
 	 * @param InMipLevel Specified mip level.
 	 * @param InSelection Requested tile selection.
 	 * @return true if the selection is contained
@@ -195,7 +202,7 @@ struct FImgMediaMipTiles final
 
 	/**
 	 * Get a copy of the selection at the specified mip level.
-	 * 
+	 *
 	 * @param InMipLevel Specified mip level.
 	 * @param OutSelection Returned tile selection.
 	 * @return true if mip level was found
@@ -437,7 +444,7 @@ public:
 	virtual bool ReadFrame(int32 FrameId, const TMap<int32, FImgMediaTileSelection>& InMipTiles, TSharedPtr<FImgMediaFrame, ESPMode::ThreadSafe> OutFrame) = 0;
 
 	/**
-	 * Mark Frame to be canceled based on Frame number. Typically this will be 
+	 * Mark Frame to be canceled based on Frame number. Typically this will be
 	 * @param FrameNumber used for frame lookup.
 	 *
 	 */
@@ -459,7 +466,7 @@ public:
 	/**
 	 * Used in case reader needs to do some processing once per frame.
 	 * Example: ExrImgMediaReaderGpu which returns unused memory to memory pool.
-	 * 
+	 *
 	 */
 	virtual void OnTick() {};
 
