[ImgMedia] Fixing non-safe game tickable created in async load.

diff --git a/Engine/Plugins/Media/ImgMedia/Source/ImgMedia/Internal/Assets/ImgMediaMipMapInfo.h b/Engine/Plugins/Media/ImgMedia/Source/ImgMedia/Internal/Assets/ImgMediaMipMapInfo.h
--- a/Engine/Plugins/Media/ImgMedia/Source/ImgMedia/Internal/Assets/ImgMediaMipMapInfo.h
+++ b/Engine/Plugins/Media/ImgMedia/Source/ImgMedia/Internal/Assets/ImgMediaMipMapInfo.h
@@ -24,9 +24,9 @@ struct FSequenceInfo
 	/** Name of this sequence. */
 	FName Name;
 	/** Pixel dimensions of this sequence. */
-	FIntPoint Dim;
+	FIntPoint Dim = FIntPoint::ZeroValue;
 	/** Number of mip levels. */
-	int32 NumMipLevels;
+	int32 NumMipLevels = 0;
 	/** Tiling description. */
 	FMediaTextureTilingDescription TilingDescription;
 	/** 
@@ -303,6 +303,7 @@ public:
 	virtual void Tick(float DeltaTime) override;
 	virtual TStatId GetStatId() const override { RETURN_QUICK_DECLARE_CYCLE_STAT(FImgMediaMipMapInfo, STATGROUP_Tickables); }
 	virtual bool IsTickableInEditor() const override { return true; }
+	virtual ETickableTickType GetTickableTickType() const override;
 
 protected:
 
@@ -338,4 +339,7 @@ protected:
 
 	/** OnEndFrame core delegate handle. */
 	FDelegateHandle OnEndFrameHandle;
+
+	/** Keep track of the ticking type, if enabled, it should be ETickableTickType::Conditional. */
+	bool bTickEnabled = false;
 };
diff --git a/Engine/Plugins/Media/ImgMedia/Source/ImgMedia/Private/Assets/ImgMediaMipMapInfo.cpp b/Engine/Plugins/Media/ImgMedia/Source/ImgMedia/Private/Assets/ImgMediaMipMapInfo.cpp
--- a/Engine/Plugins/Media/ImgMedia/Source/ImgMedia/Private/Assets/ImgMediaMipMapInfo.cpp
+++ b/Engine/Plugins/Media/ImgMedia/Source/ImgMedia/Private/Assets/ImgMediaMipMapInfo.cpp
@@ -896,7 +896,8 @@ namespace {
 } //end anonymous namespace
 
 FImgMediaMipMapInfo::FImgMediaMipMapInfo()
-	: bIsCacheValid(false)
+	: FTickableGameObject(ETickableTickType::Never)
+	, bIsCacheValid(false)
 	, OnEndFrameHandle()
 {
 }
@@ -933,6 +934,12 @@ void FImgMediaMipMapInfo::AddObject(AActor* InActor, TWeakPtr<FMediaTextureTrack
 			UE_LOG(LogImgMedia, Verbose, TEXT("Added tracked %s for %s"), *MeshComponent->GetName(), *InActor->GetName());
 		}
 
+		if (!bTickEnabled)
+		{
+			SetTickableTickType(ETickableTickType::Conditional);
+			bTickEnabled = true;
+		}
+
 		SubscribeEndFrame();
 	}
 }
@@ -1173,6 +1180,11 @@ void FImgMediaMipMapInfo::Tick(float DeltaTime)
 	}
 }
 
+ETickableTickType FImgMediaMipMapInfo::GetTickableTickType() const
+{
+	return bTickEnabled ? ETickableTickType::Conditional : ETickableTickType::Never;
+}
+
 void FImgMediaMipMapInfo::SubscribeEndFrame()
 {
 	if (!OnEndFrameHandle.IsValid())
