[ImgMedia] Handle global cache frame eviction while scrubbing.

This regression was introduced in 5.6 (CL 40284436 )

diff --git a/Engine/Plugins/Media/ImgMedia/Source/ImgMedia/Internal/Loader/ImgMediaLoader.h b/Engine/Plugins/Media/ImgMedia/Source/ImgMedia/Internal/Loader/ImgMediaLoader.h
--- a/Engine/Plugins/Media/ImgMedia/Source/ImgMedia/Internal/Loader/ImgMediaLoader.h
+++ b/Engine/Plugins/Media/ImgMedia/Source/ImgMedia/Internal/Loader/ImgMediaLoader.h
@@ -658,6 +658,9 @@ private:
 private:
 	bool RequestFrame(int32 FrameNumber, float PlayRate, bool Loop);
 
+	/** When using a global cache, handles the case when the last requested frame might have been evicted. */
+	void HandleGlobalCacheFrameEviction();
+	
 	/** Index of the previously requested frame. */
 	int32 LastRequestedFrame;
 	/** Counter used to retry loading frames again. */
diff --git a/Engine/Plugins/Media/ImgMedia/Source/ImgMedia/Private/Loader/ImgMediaLoader.cpp b/Engine/Plugins/Media/ImgMedia/Source/ImgMedia/Private/Loader/ImgMediaLoader.cpp
--- a/Engine/Plugins/Media/ImgMedia/Source/ImgMedia/Private/Loader/ImgMediaLoader.cpp
+++ b/Engine/Plugins/Media/ImgMedia/Source/ImgMedia/Private/Loader/ImgMediaLoader.cpp
@@ -568,29 +568,10 @@ IMediaSamples::EFetchBestSampleResult FImgMediaLoader::FetchBestVideoSampleForTi
 		else
 		{
 			// We did not get a frame...
-			// Could we have lost a frame that we previously loaded
-			// due to the global cache being full?
+			// Could we have lost a frame that we previously loaded due to the global cache being full?
 			if (UseGlobalCache)
 			{
-				// Did we get a frame previously?
-				if (LastRequestedFrame != INDEX_NONE)
-				{
-					// Are we loading any frames?
-					if ((PendingFrameNumbers.Num() == 0) && (QueuedFrameNumbers.Num() == 0))
-					{
-						// Nope...
-						// Wait for this to happen for one more frame.
-						// If we have a 1 image sequence, we might have just missed the frame
-						// so try again next time.
-						RetryCount++;
-						if (RetryCount > 1)
-						{
-							UE_LOG(LogImgMedia, Error, TEXT("Reloading frames. The global cache may be too small."));
-							LastRequestedFrame = INDEX_NONE;
-							RetryCount = 0;
-						}
-					}
-				}
+				HandleGlobalCacheFrameEviction();
 			}
 		}
 		CSV_EVENT(ImgMedia, TEXT("LoaderFetchMiss %" INT64_FMT "-%" INT64_FMT), UnwrappedStartIndex % NumFrames, UnwrappedEndIndex % NumFrames);
@@ -650,6 +631,14 @@ bool FImgMediaLoader::PeekVideoSampleTime(FMediaTimeStamp &TimeStamp, bool bIsLo
 				// Request the missing frame to avoid blocking.
 				RequestFrame(PeekFrameNumber, PlayRate, bIsLoopingEnabled);
 
+				// We did not get a frame...
+				// Could we have lost a frame that we previously loaded due to the global cache being full?
+				if (UseGlobalCache)
+				{
+					FScopeLock Lock(&CriticalSection);
+					HandleGlobalCacheFrameEviction();
+				}
+
 				// Need to invalidate the time stamp because FMediaPlayerFacade::UpdateSeekStatus ignores the return value.
 				TimeStamp.Invalidate();
 			}
@@ -782,6 +771,28 @@ bool FImgMediaLoader::RequestFrame(int32 FrameNumber, float PlayRate, bool Loop)
 	return true;
 }
 
+void FImgMediaLoader::HandleGlobalCacheFrameEviction()
+{
+	// Did we get a frame previously?
+	if (LastRequestedFrame != INDEX_NONE)
+	{
+		// Are we loading any frames?
+		if ((PendingFrameNumbers.Num() == 0) && (QueuedFrameNumbers.Num() == 0))
+		{
+			// Nope...
+			// Wait for this to happen for one more frame.
+			// If we have a 1 image sequence, we might have just missed the frame
+			// so try again next time.
+			RetryCount++;
+			if (RetryCount > 1)
+			{
+				UE_LOG(LogImgMedia, Error, TEXT("Reloading frames. The global cache may be too small."));
+				LastRequestedFrame = INDEX_NONE;
+				RetryCount = 0;
+			}
+		}
+	}
+}
 
 /* FImgMediaLoader implementation
  *****************************************************************************/
