[ImgMedia] Fix out of memory crash caused by unprotected large frame gaps.

Introducing a maximum frame gap cvar "ImgMedia.MaxFrameSequenceGap" (default 10000) to protect against accidentally running out of memory with incorrect file names.

Example of filenames that would cause out of memory:
frame_prefix_0000000000.exr
frame_prefix_2147483647.exr

or

frame_prefix_00000000001.exr
frame_prefix_10000000002.exr

Both case would lead to a gap of 2147483647 frames which leads to a 32 GB allocation or causing index beyond int32 for the internal frame name array.

diff --git a/Engine/Plugins/Media/ImgMedia/Source/ImgMedia/Private/Loader/ImgMediaLoader.cpp b/Engine/Plugins/Media/ImgMedia/Source/ImgMedia/Private/Loader/ImgMediaLoader.cpp
--- a/Engine/Plugins/Media/ImgMedia/Source/ImgMedia/Private/Loader/ImgMediaLoader.cpp
+++ b/Engine/Plugins/Media/ImgMedia/Source/ImgMedia/Private/Loader/ImgMediaLoader.cpp
@@ -44,6 +44,11 @@ static TAutoConsoleVariable<int32> CVarImgMediaFrameInvalidationMaxCount(
 	2,
 	TEXT("Maximum number of cached frames that can be invalidated when missing the latest mips/tiles."));
 
+static TAutoConsoleVariable<int32> CVarImgMediaMaxFrameSequenceGap(
+	TEXT("ImgMedia.MaxFrameSequenceGap"),
+	10000,
+	TEXT("Maximum number of frames that will be filled by empty frames in a sequence when there is a gap in the frame indices of the filenames."));
+
 namespace ImgMediaLoader
 {
 	void CheckAndUpdateImgDimensions(FIntPoint& InOutSequenceDim, const FIntPoint& InNewDim)
@@ -1058,7 +1063,7 @@ void FImgMediaLoader::FindFiles(const FString& SequencePath, TArray<FString>& Ou
 	}
 
 	int32 LastIndex = -1;
-	FString EmptyString;
+	const int32 MaxSequenceGap = CVarImgMediaMaxFrameSequenceGap.GetValueOnAnyThread();
 	for (const auto& Pair : SortedFoundFiles)
 	{
 		// Can we fill in gaps in the sequence?
@@ -1070,9 +1075,22 @@ void FImgMediaLoader::FindFiles(const FString& SequencePath, TArray<FString>& Ou
 			// Fill in any gaps from the last frame.
 			if ((LastIndex != -1) && (LastIndex < ThisIndex - 1))
 			{
-				for (int32 Index = LastIndex + 1; Index < ThisIndex; ++Index)
+				// Protect against large gaps that could lead to out of memory allocation.
+				const int32 SequenceGap = (ThisIndex - 1) - LastIndex;
+				if (SequenceGap < MaxSequenceGap)
+				{
+					for (int32 Index = 0; Index < SequenceGap; ++Index)
+					{
+						const FString EmptyString;
+						OutputPaths.Add(EmptyString);
+					}
+				}
+				else
 				{
-					OutputPaths.Add(EmptyString);
+					UE_LOG(LogImgMedia, Warning,
+						TEXT("Loader %p, sequence \"%s\": Ignoring large frame gap (%d frames) from frame %d to %d.")
+						TEXT(" If this is expected, adjust cvar \"ImgMedia.MaxFrameSequenceGap\" (current: %d) accordingly."),
+						this, *SequencePath, SequenceGap, LastIndex, ThisIndex, MaxSequenceGap);
 				}
 			}
 
