[ImgMedia] Adding Sliding Time Interval pre-caching strategy to alleviate global cache contention.

diff --git a/Engine/Plugins/Media/ImgMedia/Source/ImgMedia/Private/Loader/ImgMediaLoader.cpp b/Engine/Plugins/Media/ImgMedia/Source/ImgMedia/Private/Loader/ImgMediaLoader.cpp
--- a/Engine/Plugins/Media/ImgMedia/Source/ImgMedia/Private/Loader/ImgMediaLoader.cpp
+++ b/Engine/Plugins/Media/ImgMedia/Source/ImgMedia/Private/Loader/ImgMediaLoader.cpp
@@ -962,7 +962,7 @@ bool FImgMediaLoader::LoadSequence(const FString& SequencePath, const FFrameRate
 	const FPlatformMemoryStats Stats = FPlatformMemory::GetStats();
 
 	// We always need to cache at least 2 frames - current and the future frame.
-	const int32 MinNumberOfFramesToLoad = 2;
+	constexpr int32 MinNumberOfFramesToLoad = 2;
 
 	// Are we using the smart cache?
 	if (SmartCacheSettings.bIsEnabled)
@@ -982,13 +982,23 @@ bool FImgMediaLoader::LoadSequence(const FString& SequencePath, const FFrameRate
 	{
 		const SIZE_T DesiredCacheSize = Settings->CacheSizeGB * 1024 * 1024 * 1024;
 		const SIZE_T CacheSize = FMath::Clamp(DesiredCacheSize, (SIZE_T)0, (SIZE_T)Stats.AvailablePhysical);
+		const int32 NumImages = GetNumImages();
+		const int32 MaxFramesToLoad = FMath::Min((int32)(CacheSize / UncompressedSize), NumImages);
+		const int32 MinFramesToLoad = FMath::Min(MinNumberOfFramesToLoad, NumImages);
 
-		const int32 MaxFramesToLoad = (int32)(CacheSize / UncompressedSize);
-		int32 NumImages = GetNumImages();
-		int32 MinFrames = FMath::Min(MinNumberOfFramesToLoad, NumImages);
-		NumFramesToLoad = FMath::Clamp(MaxFramesToLoad, MinFrames, NumImages);
-		const float LoadBehindScale = FMath::Clamp(Settings->CacheBehindPercentage, 0.0f, 100.0f) / 100.0f;
+		if (Settings->bEnablePrecachingTimeInterval)
+		{
+			// Desired number of frames specified by a time interval, clamped by the max desired size.
+			NumFramesToLoad = FMath::Clamp(Settings->PrecachingTimeInterval / SequenceFrameRate.AsInterval(), MinFramesToLoad, MaxFramesToLoad);
+		}
+		else
+		{
+			// Desired number of frames specified by the max desired size.
+			NumFramesToLoad = FMath::Max(MaxFramesToLoad, MinFramesToLoad);
+		}
 
+		// Apply look ahead/behind ratio.
+		const float LoadBehindScale = FMath::Clamp(Settings->CacheBehindPercentage, 0.0f, 100.0f) / 100.0f;
 		NumLoadBehind = (int32)(LoadBehindScale * NumFramesToLoad);
 		NumLoadAhead = FMath::CeilToInt32((1.0f - LoadBehindScale) * NumFramesToLoad);
 	}
diff --git a/Engine/Plugins/Media/ImgMedia/Source/ImgMediaFactory/Private/ImgMediaSettings.cpp b/Engine/Plugins/Media/ImgMedia/Source/ImgMediaFactory/Private/ImgMediaSettings.cpp
--- a/Engine/Plugins/Media/ImgMedia/Source/ImgMediaFactory/Private/ImgMediaSettings.cpp
+++ b/Engine/Plugins/Media/ImgMedia/Source/ImgMediaFactory/Private/ImgMediaSettings.cpp
@@ -13,6 +13,8 @@ UImgMediaSettings::UImgMediaSettings()
 	, BandwidthThrottlingEnabled(true)
 	, CacheBehindPercentage(25)
 	, CacheSizeGB(1.0f)
+	, bEnablePrecachingTimeInterval(false)
+	, PrecachingTimeInterval(0.2f)
 	, CacheThreads(2)
 	, CacheThreadStackSizeKB(128)
 	, GlobalCacheSizeGB(1.0f)
diff --git a/Engine/Plugins/Media/ImgMedia/Source/ImgMediaFactory/Public/ImgMediaSettings.h b/Engine/Plugins/Media/ImgMedia/Source/ImgMediaFactory/Public/ImgMediaSettings.h
--- a/Engine/Plugins/Media/ImgMedia/Source/ImgMediaFactory/Public/ImgMediaSettings.h
+++ b/Engine/Plugins/Media/ImgMedia/Source/ImgMediaFactory/Public/ImgMediaSettings.h
@@ -38,10 +38,29 @@ public:
 	UPROPERTY(config, EditAnywhere, Category=Caching, meta=(ClampMin=0.0, ClampMax=100.0))
 	float CacheBehindPercentage;
 
-	/** Maximum size of the look-ahead cache (in GB; default = 1 GB). */
+	/**
+	 * Maximum memory size of the pre-caching sliding interval (in GB; default = 1 GB).
+	 * The pre-caching sliding interval can further be configured by specifying
+	 * a time interval (see PrecachingTimeInterval). In that case, the effective
+	 * interval will be the minimum of those two settings (size in GB or time span).
+	 */
 	UPROPERTY(config, EditAnywhere, Category=Caching, meta=(ClampMin=0))
 	float CacheSizeGB;
 
+	/**
+	 * Enables pre-caching time interval strategy.
+	 */
+	UPROPERTY(config, EditAnywhere, Category=Caching, meta=(InlineEditConditionToggle))
+	bool bEnablePrecachingTimeInterval;
+
+	/**
+	 * Specify the pre-caching sliding time interval, in seconds.
+	 * The player will pre-cache the specified number of frames corresponding to the time interval relative to the play head.
+	 * The cache behind percentage still applies to determine how much frames to pre-cache behind vs ahead of the play head.
+	 */
+	UPROPERTY(EditAnywhere, Category=Caching, meta=(ClampMin=0, EditCondition="bEnablePrecachingTimeInterval"))
+	float PrecachingTimeInterval;
+	
 	/** Maximum number of image caching threads to use (0 = number of cores, default = 2). */
 	UPROPERTY(config, EditAnywhere, Category=Caching, meta=(ClampMin=0))
 	int32 CacheThreads;
