[ImgMedia] Fixing RHIValidation failures for EXR Reader Gpu.

#Synced-CL 47720872

diff --git a/Engine/Plugins/Media/ImgMedia/Source/ImgMedia/Private/Readers/ExrImgMediaReaderGpu.cpp b/Engine/Plugins/Media/ImgMedia/Source/ImgMedia/Private/Readers/ExrImgMediaReaderGpu.cpp
--- a/Engine/Plugins/Media/ImgMedia/Source/ImgMedia/Private/Readers/ExrImgMediaReaderGpu.cpp
+++ b/Engine/Plugins/Media/ImgMedia/Source/ImgMedia/Private/Readers/ExrImgMediaReaderGpu.cpp
@@ -183,6 +183,8 @@ FExrImgMediaReader::EReadResult FExrImgMediaReaderGpu::ReadMip
 					RHI_BREADCRUMB_EVENT_STAT_F(RHICmdList, ExrImgMediaReaderGpu_CopyUploadBuffer, "ExrReaderGpu.StartCopy", "ExrReaderGpu.StartCopy %d", FrameId);
 					SCOPED_GPU_STAT(RHICmdList, ExrImgMediaReaderGpu_CopyUploadBuffer);
 
+					RHICmdList.Transition(FRHITransitionInfo(BufferData->ShaderAccessBufferRef, ERHIAccess::SRVMask, ERHIAccess::CopyDest));
+
 					if (BufferRegionsToCopy.IsEmpty())
 					{
 						RHICmdList.CopyBufferRegion(BufferData->ShaderAccessBufferRef, 0, BufferData->UploadBufferRef, 0, BufferData->ShaderAccessBufferRef->GetSize());
@@ -194,6 +196,8 @@ FExrImgMediaReader::EReadResult FExrImgMediaReaderGpu::ReadMip
 							RHICmdList.CopyBufferRegion(BufferData->ShaderAccessBufferRef, Region.X, BufferData->UploadBufferRef, Region.X, Region.Y);
 						}
 					}
+
+					RHICmdList.Transition(FRHITransitionInfo(BufferData->ShaderAccessBufferRef, ERHIAccess::CopyDest, ERHIAccess::SRVMask));
 				});
 		}
 	}
@@ -569,6 +573,9 @@ void FExrImgMediaReaderGpu::CreateSampleConverterCallback(TSharedPtr<FExrMediaTe
 			RHICmdList.EndRenderPass();
 		};
 
+		// Using the same transition pattern as ConvertTextureToOutput.
+		RHICmdList.Transition(FRHITransitionInfo(RenderTargetTextureRHI, ERHIAccess::Unknown, ERHIAccess::RTV));
+
 		int32 MipToUpscale = ConverterParams.UpscaleMip;
 
 		// Upscale to all mips below the mip to upscale.
@@ -624,6 +631,8 @@ void FExrImgMediaReaderGpu::CreateSampleConverterCallback(TSharedPtr<FExrMediaTe
 			}
 		}
 
+		RHICmdList.Transition(FRHITransitionInfo(RenderTargetTextureRHI, ERHIAccess::RTV, ERHIAccess::SRVMask));
+
 		//Doesn't need further conversion so returning false.
 		return false;
 	};
@@ -687,6 +696,12 @@ FStructuredBufferPoolItemSharedPtr FExrImgMediaReaderGpu::AllocateGpuBufferFromP
 						.DetermineInitialState();
 					AllocatedBuffer->UploadBufferRef = RHICmdList.CreateBuffer(CreateDesc);
 					AllocatedBuffer->UploadBufferMapped = RHICmdList.LockBuffer(AllocatedBuffer->UploadBufferRef, 0, AllocSize, RLM_WriteOnly);
+
+					if (bExrReaderUseUploadHeap)
+					{
+						// The upload buffer will remain in CopySrc state.
+						RHICmdList.Transition(FRHITransitionInfo(AllocatedBuffer->UploadBufferRef, CreateDesc.InitialState, ERHIAccess::CopySrc));
+					}
 				}
 
 				if (bExrReaderUseUploadHeap)
