ImgMedia: Cleanup/replacement of IImageWrapper in favor of preferred FImage interface in generic image sequence reader.

* Adds new "NumBytesPerPixel" property to image frame info struct, filled when possible.

#Synced-CL 48601581

diff --git a/Engine/Plugins/Media/ImgMedia/Source/ImgMedia/ImgMedia.Build.cs b/Engine/Plugins/Media/ImgMedia/Source/ImgMedia/ImgMedia.Build.cs
--- a/Engine/Plugins/Media/ImgMedia/Source/ImgMedia/ImgMedia.Build.cs
+++ b/Engine/Plugins/Media/ImgMedia/Source/ImgMedia/ImgMedia.Build.cs
@@ -8,7 +8,6 @@ namespace UnrealBuildTool.Rules
 		{
 			DynamicallyLoadedModuleNames.AddRange(
 				new string[] {
-					"ImageWrapper",
 					"Media",
 				});
 
@@ -32,7 +31,6 @@ namespace UnrealBuildTool.Rules
 
 			PrivateIncludePathModuleNames.AddRange(
 				new string[] {
-					"ImageWrapper",
 					"Media",
 				});
 
diff --git a/Engine/Plugins/Media/ImgMedia/Source/ImgMedia/Internal/Loader/ImgMediaLoader.h b/Engine/Plugins/Media/ImgMedia/Source/ImgMedia/Internal/Loader/ImgMediaLoader.h
--- a/Engine/Plugins/Media/ImgMedia/Source/ImgMedia/Internal/Loader/ImgMediaLoader.h
+++ b/Engine/Plugins/Media/ImgMedia/Source/ImgMedia/Internal/Loader/ImgMediaLoader.h
@@ -17,7 +17,6 @@ class FImgMediaLoaderWork;
 class FImgMediaScheduler;
 class FImgMediaTextureSample;
 class FMediaTimeStamp;
-class IImageWrapperModule;
 class IImgMediaReader;
 class IQueuedWork;
 
@@ -132,15 +131,6 @@ public:
 	 */
 	void GetCompletedTimeRanges(TRangeSet<FTimespan>& OutRangeSet) const;
 
-	/**
-	 * Get the image frame at the specified time.
-	 *
-	 * @param Time The time of the image frame to get (relative to the beginning of the sequence).
-	 * @return The frame, or nullptr if the frame wasn't available yet.
-	 */
-	UE_DEPRECATED(5.4, "This method will be removed.")
-	TSharedPtr<FImgMediaTextureSample, ESPMode::ThreadSafe> GetFrameSample(FTimespan Time);
-
 	/**
 	 * Get the information string for the currently loaded image sequence.
 	 *
@@ -272,6 +262,14 @@ public:
 		return NumMipLevels; 
 	};
 
+	/**
+	 * Get the number of bytes per pixel for each image.
+	 */
+	SIZE_T GetNumBytesPerPixel() const
+	{
+		return NumBytesPerPixel;
+	}
+
 	/**
 	 * Get the minumum mip level to upscale to lower quality mips (-1 if disabled).
 	 */
@@ -593,9 +591,6 @@ private:
 	/** The currently loaded image sequence frames. */
 	TLruCache<int32, TSharedPtr<FImgMediaFrame, ESPMode::ThreadSafe>> Frames;
 
-	/** The image wrapper module to use. */
-	IImageWrapperModule& ImageWrapperModule;
-
 	/**
 	 * Paths to each image for each mip map level in the currently opened sequence.
 	 * This is an array of mip levels, and each mip level is an array of image paths.
@@ -623,9 +618,12 @@ private:
 	/** Total number of frames to load. */
 	int32 NumFramesToLoad;
 
-	/** Stores num of mip levels if the sequence is made out of files that contain all mips in one file. */
+	/** Stores number of mip levels. */
 	int32 NumMipLevels;
 
+	/** Stores number of bytes per pixel for each image. */
+	SIZE_T NumBytesPerPixel;
+
 	/** Minumum mip level to upscale. */
 	TAtomic<int32> MinimumMipLevelToUpscale;
 
diff --git a/Engine/Plugins/Media/ImgMedia/Source/ImgMedia/Private/Loader/ImgMediaLoader.cpp b/Engine/Plugins/Media/ImgMedia/Source/ImgMedia/Private/Loader/ImgMediaLoader.cpp
--- a/Engine/Plugins/Media/ImgMedia/Source/ImgMedia/Private/Loader/ImgMediaLoader.cpp
+++ b/Engine/Plugins/Media/ImgMedia/Source/ImgMedia/Private/Loader/ImgMediaLoader.cpp
@@ -170,12 +170,14 @@ FImgMediaLoader::FImgMediaLoader(const TSharedRef<FImgMediaScheduler, ESPMode::T
 	bool bInFillGapsInSequence,
 	const FImgMediaLoaderSmartCacheSettings& InSmartCacheSettings)
 	: Frames(1)
-	, ImageWrapperModule(FModuleManager::LoadModuleChecked<IImageWrapperModule>("ImageWrapper"))
 	, Initialized(false)
 	, bFillGapsInSequence(bInFillGapsInSequence)
 	, TilingDescription()
 	, NumLoadAhead(0)
 	, NumLoadBehind(0)
+	, NumFramesToLoad(0)
+	, NumMipLevels(1)
+	, NumBytesPerPixel(0)
 	, Scheduler(InScheduler)
 	, GlobalCache(InGlobalCache)
 	, MipMapInfo(InMipMapInfo)
@@ -233,7 +235,7 @@ FImgMediaLoader::~FImgMediaLoader()
 uint64 FImgMediaLoader::GetBitRate() const
 {
 	FScopeLock Lock(&CriticalSection);
-	return SequenceDim.X * SequenceDim.Y * sizeof(uint16) * 8 * SequenceFrameRate.AsDecimal();
+	return SequenceDim.X * SequenceDim.Y * NumBytesPerPixel * 8 * SequenceFrameRate.AsDecimal();
 }
 
 
@@ -266,39 +268,6 @@ void FImgMediaLoader::GetCompletedTimeRanges(TRangeSet<FTimespan>& OutRangeSet)
 	FrameNumbersToTimeRanges(CompletedFrames, OutRangeSet);
 }
 
-
-//note: use with V1 player version only!
-TSharedPtr<FImgMediaTextureSample, ESPMode::ThreadSafe> FImgMediaLoader::GetFrameSample(FTimespan Time)
-{
-	const uint32 FrameIndex = TimeToFrameNumber(Time);
-
-	if (FrameIndex == INDEX_NONE)
-	{
-		return nullptr;
-	}
-
-	FScopeLock ScopeLock(&CriticalSection);
-
-	const TSharedPtr<FImgMediaFrame, ESPMode::ThreadSafe>* Frame = GetCachedFrame(FrameIndex);
-
-	if (Frame == nullptr)
-	{
-		return nullptr;
-	}
-
-	const FTimespan FrameStartTime = FrameNumberToTime(FrameIndex);
-	const FTimespan NextStartTime = FrameNumberToTime(FrameIndex + 1);
-
-	auto Sample = MakeShared<FImgMediaTextureSample, ESPMode::ThreadSafe>();
-
-	ImgMediaLoader::CheckAndUpdateImgDimensions(SequenceDim, Frame->Get()->GetDim());
-	
-	Sample->Initialize(*Frame->Get(), SequenceDim, FMediaTimeStamp(FrameStartTime, 0), NextStartTime - FrameStartTime, GetNumMipLevels(), TilingDescription, SourceColorSettings);
-
-	return Sample;
-}
-
-
 const FString& FImgMediaLoader::GetImagePath(int32 FrameNumber, int32 MipLevel) const
 {
 	if ((MipLevel < 0) || (MipLevel >= ImagePaths.Num() ||
@@ -870,8 +839,9 @@ bool FImgMediaLoader::LoadSequence(const FString& SequencePath, const FFrameRate
 	}
 	else
 	{
-		Reader = MakeShareable(new FGenericImgMediaReader(ImageWrapperModule, AsShared()));
+		Reader = MakeShareable(new FGenericImgMediaReader(AsShared()));
 	}
+	
 	if (Reader.IsValid() == false)
 	{
 		UE_LOG(LogImgMedia, Error, TEXT("Reader is not valid for file %s."), *ImagePaths[0][0]);
@@ -904,6 +874,7 @@ bool FImgMediaLoader::LoadSequence(const FString& SequencePath, const FFrameRate
 			return false;
 		}
 	}
+	
 	if (FirstFrameInfo.UncompressedSize == 0)
 	{
 		UE_LOG(LogImgMedia, Error, TEXT("The first image in sequence %s does not have a valid frame size"), *SequencePath);
@@ -917,6 +888,8 @@ bool FImgMediaLoader::LoadSequence(const FString& SequencePath, const FFrameRate
 	}
 
 	NumMipLevels = FMath::Max(FirstFrameInfo.NumMipLevels, ImagePaths.Num());
+	NumBytesPerPixel = FirstFrameInfo.NumBytesPerPixel;
+
 	MinimumMipLevelToUpscale = -1;
 
 	SequenceDim = FirstFrameInfo.Dim;
@@ -1050,7 +1023,7 @@ void FImgMediaLoader::FindFiles(const FString& SequencePath, TArray<FString>& Ou
 
 	UE_LOG(LogImgMedia, Verbose, TEXT("Loader %p: Found %i image files in %s"), this, FoundFiles.Num(), *SequencePath);
 
-	// Same list of extensions as FGenericImgMediaReader::LoadFrameImage() with the addition of exr
+	// Same list of extensions as FGenericImgMediaReader::GetFrameInfoFromImage() with the addition of exr & dds
 	TSet<FString> SupportedImageFileExtensions = { TEXT("exr"), TEXT("jpg"), TEXT("jpeg"), TEXT("png"), TEXT("bmp"), TEXT("dds") };
 	SupportedImageFileExtensions.Append(FTmvImgMediaReader::GetSupportedFileExtensions());
 
diff --git a/Engine/Plugins/Media/ImgMedia/Source/ImgMedia/Private/Readers/DdsImgMediaReader.cpp b/Engine/Plugins/Media/ImgMedia/Source/ImgMedia/Private/Readers/DdsImgMediaReader.cpp
--- a/Engine/Plugins/Media/ImgMedia/Source/ImgMedia/Private/Readers/DdsImgMediaReader.cpp
+++ b/Engine/Plugins/Media/ImgMedia/Source/ImgMedia/Private/Readers/DdsImgMediaReader.cpp
@@ -3,8 +3,6 @@
 #include "DdsImgMediaReader.h"
 #include "ImgMediaPrivate.h"
 
-#include "IImageWrapper.h"
-#include "IImageWrapperModule.h"
 #include "ImageCoreUtils.h"
 #include "Loader/ImgMediaLoader.h"
 #include "Logging/StructuredLog.h"
diff --git a/Engine/Plugins/Media/ImgMedia/Source/ImgMedia/Private/Readers/ExrImgMediaReader.cpp b/Engine/Plugins/Media/ImgMedia/Source/ImgMedia/Private/Readers/ExrImgMediaReader.cpp
--- a/Engine/Plugins/Media/ImgMedia/Source/ImgMedia/Private/Readers/ExrImgMediaReader.cpp
+++ b/Engine/Plugins/Media/ImgMedia/Source/ImgMedia/Private/Readers/ExrImgMediaReader.cpp
@@ -358,6 +358,7 @@ bool FExrImgMediaReader::GetInfo(const FString& FilePath, FImgMediaFrameInfo& Ou
 	OutInfo.Srgb = false;
 	OutInfo.UncompressedSize = HeaderReader.GetUncompressedSize();
 	OutInfo.NumChannels = HeaderReader.GetNumChannels();
+	OutInfo.NumBytesPerPixel = sizeof(uint16) * OutInfo.NumChannels;
 	OutInfo.NumMipLevels = 1;
 
 	int32 CustomFormat = 0;
diff --git a/Engine/Plugins/Media/ImgMedia/Source/ImgMedia/Private/Readers/GenericImgMediaReader.cpp b/Engine/Plugins/Media/ImgMedia/Source/ImgMedia/Private/Readers/GenericImgMediaReader.cpp
--- a/Engine/Plugins/Media/ImgMedia/Source/ImgMedia/Private/Readers/GenericImgMediaReader.cpp
+++ b/Engine/Plugins/Media/ImgMedia/Source/ImgMedia/Private/Readers/GenericImgMediaReader.cpp
@@ -3,81 +3,83 @@
 #include "GenericImgMediaReader.h"
 #include "ImgMediaPrivate.h"
 
-#include "IImageWrapper.h"
-#include "IImageWrapperModule.h"
+#include "ImageUtils.h"
 #include "Loader/ImgMediaLoader.h"
 #include "Misc/FileHelper.h"
 #include "Misc/Paths.h"
 
 
-TSharedPtr<IImageWrapper> FGenericImgMediaReader::LoadFrameImage(const FString& ImagePath, TArray64<uint8>& OutBuffer, FImgMediaFrameInfo& OutInfo)
+bool FGenericImgMediaReader::IsSupportedFormat(const ERawImageFormat::Type InFormat) const
 {
-	// load image into buffer
-	if (!FFileHelper::LoadFileToArray(OutBuffer, *ImagePath))
+	switch (InFormat)
 	{
-		UE_LOG(LogImgMedia, Warning, TEXT("FGenericImgMediaReader: Failed to load %s"), *ImagePath);
-		return nullptr;
+	case ERawImageFormat::BGRA8:
+	case ERawImageFormat::RGBA16:
+	case ERawImageFormat::RGBA16F:
+		return true;
+	default:
+		// Other formats do not currently have a EMediaTextureSampleFormat match.
+		return false;
 	}
+}
 
-	// determine image format
-	EImageFormat ImageFormat;
+bool FGenericImgMediaReader::GetFrameInfoFromImage(const FString& ImagePath, const FImageInfo& ImageInfo, FImgMediaFrameInfo& OutInfo)
+{
+	if (!IsSupportedFormat(ImageInfo.Format))
+	{
+		return false;
+	}
 
 	const FString Extension = FPaths::GetExtension(ImagePath).ToLower();
 
 	if (Extension == TEXT("bmp"))
 	{
-		ImageFormat = EImageFormat::BMP;
 		OutInfo.FormatName = TEXT("BMP");
 	}
 	else if ((Extension == TEXT("jpg")) || (Extension == TEXT("jpeg")))
 	{
-		ImageFormat = EImageFormat::JPEG;
 		OutInfo.FormatName = TEXT("JPEG");
 	}
 	else if (Extension == TEXT("png"))
 	{
-		ImageFormat = EImageFormat::PNG;
 		OutInfo.FormatName = TEXT("PNG");
 	}
 	else
 	{
 		UE_LOG(LogImgMedia, Warning, TEXT("FGenericImgMediaReader: Unsupported file format in %s"), *ImagePath);
-		return nullptr;
+		return false;
 	}
 
-	// create image wrapper
-	TSharedPtr<IImageWrapper> ImageWrapper = ImageWrapperModule.CreateImageWrapper(ImageFormat);
-
-	if (!ImageWrapper.IsValid() || !ImageWrapper->SetCompressed(OutBuffer.GetData(), OutBuffer.Num()))
+	if (ImageInfo.NumSlices > 1)
 	{
-		UE_LOG(LogImgMedia, Warning, TEXT("FGenericImgMediaReader: Failed to create image wrapper for %s"), *ImagePath);
-		return nullptr;
+		UE_LOG(LogImgMedia, Warning, TEXT("FGenericImgMediaReader: Unsupported texture arrays or 3D textures in %s"), *ImagePath);
+		return false;
 	}
 
 	// get file info
 	const UImgMediaSettings* Settings = GetDefault<UImgMediaSettings>();
 	OutInfo.CompressionName = TEXT("");
-	OutInfo.Dim.X = ImageWrapper->GetWidth();
-	OutInfo.Dim.Y = ImageWrapper->GetHeight();
-	OutInfo.UncompressedSize = (SIZE_T)OutInfo.Dim.X * OutInfo.Dim.Y * 4;
+	OutInfo.Dim.X = ImageInfo.GetWidth();
+	OutInfo.Dim.Y = ImageInfo.GetHeight();
+	OutInfo.UncompressedSize = ImageInfo.GetImageSizeBytes();
 	OutInfo.NumMipLevels = 1;
 	OutInfo.FrameRate = Settings->DefaultFrameRate;
-	OutInfo.Srgb = true;
-	OutInfo.NumChannels = 4;
+	OutInfo.Srgb = (ImageInfo.GetGammaSpace() == EGammaSpace::sRGB);
+	OutInfo.NumChannels = ERawImageFormat::NumChannels(ImageInfo.Format);
+	OutInfo.NumBytesPerPixel = static_cast<SIZE_T>(ImageInfo.GetBytesPerPixel());
 	OutInfo.bHasTiles = false;
 	OutInfo.TileDimensions = OutInfo.Dim;
 	OutInfo.NumTiles = FIntPoint(1, 1);
 	OutInfo.TileBorder = 0;
 
-	return ImageWrapper;
+	return true;
 }
 
 /* FGenericImgMediaReader structors
  *****************************************************************************/
 
-FGenericImgMediaReader::FGenericImgMediaReader(IImageWrapperModule& InImageWrapperModule, const TSharedRef<FImgMediaLoader, ESPMode::ThreadSafe>& InLoader)
-	: ImageWrapperModule(InImageWrapperModule)
-	, LoaderPtr(InLoader)
+FGenericImgMediaReader::FGenericImgMediaReader(const TSharedRef<FImgMediaLoader, ESPMode::ThreadSafe>& InLoader)
+	: LoaderPtr(InLoader)
 { }
 
 
@@ -86,10 +88,13 @@ FGenericImgMediaReader::FGenericImgMediaReader(IImageWrapperModule& InImageWrapp
 
 bool FGenericImgMediaReader::GetFrameInfo(const FString& ImagePath, FImgMediaFrameInfo& OutInfo)
 {
-	TArray64<uint8> InputBuffer;
-	TSharedPtr<IImageWrapper> ImageWrapper = LoadFrameImage(ImagePath, InputBuffer, OutInfo);
+	FImage Image;
+	if(FImageUtils::LoadImage(*ImagePath, Image))
+	{
+		return GetFrameInfoFromImage(ImagePath, Image, OutInfo);
+	}
 
-	return ImageWrapper.IsValid();
+	return false;
 }
 
 
@@ -107,6 +112,7 @@ bool FGenericImgMediaReader::ReadFrame(int32 FrameId, const TMap<int32, FImgMedi
 	}
 
 	SIZE_T BufferDataOffset = 0;
+	const SIZE_T NumBytesPerPixel = Loader->GetNumBytesPerPixel();
 	const int32 NumMipLevels = Loader->GetNumMipLevels();
 	const FIntPoint BaseLevelDim = Loader->GetSequenceDim();
 	FIntPoint CurrenDim = BaseLevelDim;
@@ -125,27 +131,26 @@ bool FGenericImgMediaReader::ReadFrame(int32 FrameId, const TMap<int32, FImgMedi
 				// Load image.
 				const FString& ImagePath = Loader->GetImagePath(FrameId, CurrentMipLevel);
 
-				TArray64<uint8> InputBuffer;
-				FImgMediaFrameInfo Info;
-				TSharedPtr<IImageWrapper> ImageWrapper = LoadFrameImage(ImagePath, InputBuffer, Info);
+				FImage Image;
+				const bool bSuccess = FImageUtils::LoadImage(*ImagePath, Image);
 
-				if (!ImageWrapper.IsValid())
+				if (!bSuccess)
 				{
 					UE_LOG(LogImgMedia, Warning, TEXT("FGenericImgMediaReader: Failed to load image %s"), *ImagePath);
 					return false;
 				}
 
-				TArray64<uint8> RawData;
-				if (!ImageWrapper->GetRaw(ERGBFormat::BGRA, 8, RawData))
+				FImgMediaFrameInfo Info;
+				if (!GetFrameInfoFromImage(ImagePath, Image, Info))
 				{
-					UE_LOG(LogImgMedia, Warning, TEXT("FGenericImgMediaReader: Failed to get image data for %s"), *ImagePath);
+					UE_LOG(LogImgMedia, Warning, TEXT("FGenericImgMediaReader: Failed to obtain frame info, cancelling %s read."), *ImagePath);
 					return false;
 				}
 
 				// Create the full buffer for data.
 				if (!OutFrame->Data.IsValid())
 				{
-					SIZE_T AllocSize = (SIZE_T)Info.Dim.X * Info.Dim.Y * 4;
+					SIZE_T AllocSize = Image.GetSliceSizeBytes();
 					
 					if (NumMipLevels > 1)
 					{
@@ -155,7 +160,7 @@ bool FGenericImgMediaReader::ReadFrame(int32 FrameId, const TMap<int32, FImgMedi
 						Info.Dim = BaseLevelDim;
 						Info.NumMipLevels = NumMipLevels;
 
-						const SIZE_T SizeMip0 = (SIZE_T)BaseLevelDim.X * BaseLevelDim.Y * 4;
+						const SIZE_T SizeMip0 = (SIZE_T)BaseLevelDim.X * BaseLevelDim.Y * Image.GetBytesPerPixel();
 						Info.UncompressedSize = SizeMip0;
 						for (int32 Level = 1; Level < Info.NumMipLevels; Level++)
 						{
@@ -170,19 +175,39 @@ bool FGenericImgMediaReader::ReadFrame(int32 FrameId, const TMap<int32, FImgMedi
 					OutFrame->SetInfo(Info);
 					OutFrame->Data = MakeShareable(Buffer, [](void* ObjectToDelete) { FMemory::Free(ObjectToDelete); });
 					OutFrame->MipTilesPresent.Reset();
-					OutFrame->Format = EMediaTextureSampleFormat::CharBGRA;
-					OutFrame->Stride = Info.Dim.X * 4;
+					OutFrame->Stride = BaseLevelDim.X * NumBytesPerPixel;
+					switch (Image.Format)
+					{
+					case ERawImageFormat::BGRA8: OutFrame->Format = EMediaTextureSampleFormat::CharBGRA; break;
+					case ERawImageFormat::RGBA16: OutFrame->Format = EMediaTextureSampleFormat::RGBA16; break;
+					case ERawImageFormat::RGBA16F: OutFrame->Format = EMediaTextureSampleFormat::FloatRGBA; break;
+					default:
+						UE_LOG(LogImgMedia, Warning, TEXT("FGenericImgMediaReader: Unsupported image format, cancelling %s read."), *ImagePath);
+						return false;
+					}
+				}
+
+				if (Image.GetBytesPerPixel() != NumBytesPerPixel)
+				{
+					UE_LOG(LogImgMedia, Warning, TEXT("FGenericImgMediaReader: Mismatched bytes per pixel, cancelling %s read."), *ImagePath);
+					return false;
+				}
+
+				if ((Image.GetWidth() != CurrenDim.X) || (Image.GetHeight() != CurrenDim.Y))
+				{
+					UE_LOG(LogImgMedia, Warning, TEXT("FGenericImgMediaReader: Mismatched expected image size, cancelling %s read."), *ImagePath);
+					return false;
 				}
 
 				// Copy data to our buffer with the right mip level offset
-				FMemory::Memcpy((void*)((uint8*)OutFrame->Data.Get() + BufferDataOffset), RawData.GetData(), RawData.Num());
+				FMemory::Memcpy((void*)((uint8*)OutFrame->Data.Get() + BufferDataOffset), Image.GetPixelPointer(0, 0), Image.GetSliceSizeBytes());
 				OutFrame->MipTilesPresent.Emplace(CurrentMipLevel, CurrentSelection);
 				OutFrame->NumTilesRead++;
 			}
 		}
 
 		// Next level.
-		BufferDataOffset += (CurrenDim.X * CurrenDim.Y * 4);
+		BufferDataOffset += (CurrenDim.X * CurrenDim.Y * NumBytesPerPixel);
 		CurrenDim /= 2;
 	}
 
diff --git a/Engine/Plugins/Media/ImgMedia/Source/ImgMedia/Private/Readers/GenericImgMediaReader.h b/Engine/Plugins/Media/ImgMedia/Source/ImgMedia/Private/Readers/GenericImgMediaReader.h
--- a/Engine/Plugins/Media/ImgMedia/Source/ImgMedia/Private/Readers/GenericImgMediaReader.h
+++ b/Engine/Plugins/Media/ImgMedia/Source/ImgMedia/Private/Readers/GenericImgMediaReader.h
@@ -7,9 +7,13 @@
 
 class FImgMediaLoader;
 class FRgbaInputFile;
-class IImageWrapper;
-class IImageWrapperModule;
+struct FImageInfo;
+struct FImgMediaFrameInfo;
 
+namespace ERawImageFormat
+{
+	enum Type : uint8;
+}
 
 /**
  * Implements a reader for various image sequence formats.
@@ -21,10 +25,8 @@ public:
 
 	/**
 	 * Create and initialize a new instance.
-	 *
-	 * @param InImageWrapperModule The image wrapper module to use.
 	 */
-	FGenericImgMediaReader(IImageWrapperModule& InImageWrapperModule, const TSharedRef<FImgMediaLoader, ESPMode::ThreadSafe>& InLoader);
+	FGenericImgMediaReader(const TSharedRef<FImgMediaLoader, ESPMode::ThreadSafe>& InLoader);
 
 	/** Virtual destructor. */
 	virtual ~FGenericImgMediaReader() = default;
@@ -39,11 +41,12 @@ public:
 	virtual void UncancelFrame(int32 FrameNumber) override {};
 
 private:
+	/** Checks if the format is supported. */
+	bool IsSupportedFormat(const ERawImageFormat::Type InFormat) const;
 
-	TSharedPtr<IImageWrapper> LoadFrameImage(const FString& ImagePath, TArray64<uint8>& OutBuffer, FImgMediaFrameInfo& OutInfo);
+	/** Extract metadata from the loaded image. */
+	bool GetFrameInfoFromImage(const FString& ImagePath, const FImageInfo& Image, FImgMediaFrameInfo& OutInfo);
 
-	/** The image wrapper module. */
-	IImageWrapperModule& ImageWrapperModule;
 	/** Our parent loader. */
 	TWeakPtr<FImgMediaLoader, ESPMode::ThreadSafe> LoaderPtr;
 };
diff --git a/Engine/Plugins/Media/ImgMedia/Source/ImgMedia/Private/Readers/IImgMediaReader.h b/Engine/Plugins/Media/ImgMedia/Source/ImgMedia/Private/Readers/IImgMediaReader.h
--- a/Engine/Plugins/Media/ImgMedia/Source/ImgMedia/Private/Readers/IImgMediaReader.h
+++ b/Engine/Plugins/Media/ImgMedia/Source/ImgMedia/Private/Readers/IImgMediaReader.h
@@ -44,6 +44,9 @@ struct FImgMediaFrameInfo
 	/** Number of channels (RGB - 3 or RGBA - 4). */
 	SIZE_T NumChannels = 0;
 
+	/** Number of bytes per pixel. */
+	SIZE_T NumBytesPerPixel = 0;
+
 	/** Does this frame have tiles?. */
 	bool bHasTiles = false;
 
diff --git a/Engine/Plugins/Media/ImgMedia/Source/ImgMedia/Private/Readers/TmvImgMediaReader.cpp b/Engine/Plugins/Media/ImgMedia/Source/ImgMedia/Private/Readers/TmvImgMediaReader.cpp
--- a/Engine/Plugins/Media/ImgMedia/Source/ImgMedia/Private/Readers/TmvImgMediaReader.cpp
+++ b/Engine/Plugins/Media/ImgMedia/Source/ImgMedia/Private/Readers/TmvImgMediaReader.cpp
@@ -228,6 +228,7 @@ namespace UE::ImgMedia
 		OutInfo.FrameRate = GetDefault<UImgMediaSettings>()->DefaultFrameRate;
 		OutInfo.Srgb = true;	// todo: Transfer function and color space info.
 		OutInfo.NumChannels = InMipInfos[0].NumComponents;
+		OutInfo.NumBytesPerPixel = OutInfo.NumChannels * 2; // todo: This assumes 16 bits per channel
 		OutInfo.bHasTiles = InMipInfos[0].NumTiles.X * InMipInfos[0].NumTiles.Y > 1;
 		OutInfo.TileDimensions = FIntPoint(InMipInfos[0].TileWidth, InMipInfos[0].TileHeight);
 		OutInfo.NumTiles = InMipInfos[0].NumTiles;
