[ImgMedia] OpenExrWrapper API: Fix int32 overflow for large textures.

This is needed for transcoding large textures for TMV encoding.

#Synced-CL 48520225

diff --git a/Engine/Plugins/Media/ImgMedia/Source/OpenExrWrapper/Private/OpenExrWrapperContextReader.cpp b/Engine/Plugins/Media/ImgMedia/Source/OpenExrWrapper/Private/OpenExrWrapperContextReader.cpp
--- a/Engine/Plugins/Media/ImgMedia/Source/OpenExrWrapper/Private/OpenExrWrapperContextReader.cpp
+++ b/Engine/Plugins/Media/ImgMedia/Source/OpenExrWrapper/Private/OpenExrWrapperContextReader.cpp
@@ -119,12 +119,12 @@ FFrameRate FOpenExrHeaderReader::GetFrameRate(const FFrameRate& DefaultValue) co
 	return FFrameRate(Attribute->rational->num, Attribute->rational->denom);
 }
 
-int32 FOpenExrHeaderReader::GetUncompressedSize() const
+int64 FOpenExrHeaderReader::GetUncompressedSize() const
 {
 	const int32 TexelSize = GetPixelSize();
 	const FIntPoint Window = GetDataWindow();
 
-	return (Window.X * Window.Y * TexelSize);
+	return (static_cast<int64>(Window.X) * static_cast<int64>(Window.Y) * static_cast<int64>(TexelSize));
 }
 
 int32 FOpenExrHeaderReader::GetNumChannels() const
diff --git a/Engine/Plugins/Media/ImgMedia/Source/OpenExrWrapper/Public/OpenExrWrapper.h b/Engine/Plugins/Media/ImgMedia/Source/OpenExrWrapper/Public/OpenExrWrapper.h
--- a/Engine/Plugins/Media/ImgMedia/Source/OpenExrWrapper/Public/OpenExrWrapper.h
+++ b/Engine/Plugins/Media/ImgMedia/Source/OpenExrWrapper/Public/OpenExrWrapper.h
@@ -72,7 +72,7 @@ public:
 	UE_API const TCHAR* GetCompressionName() const;
 	UE_API FIntPoint GetDataWindow() const;
 	UE_API FFrameRate GetFrameRate(const FFrameRate& DefaultValue) const;
-	UE_API int32 GetUncompressedSize() const;
+	UE_API int64 GetUncompressedSize() const;
 	UE_API int32 GetNumChannels() const;
 	UE_API int32 GetPixelSize() const;
 	UE_API bool ContainsMips() const;
