[ImgMedia] Fix Process Exr from an Img Media Source (content browser context menu)

GetSequencePath() returns a relative path that can't necessarily be resolved with ConvertRelativePathToFull. It is more reliable to just let the media source resolve the full path itself and use the full path then for ProcessExr.

#Synced-CL 49891360

diff --git a/Engine/Plugins/Media/ImgMedia/Source/ImgMediaEditor/Private/ImgMediaEditorModule.cpp b/Engine/Plugins/Media/ImgMedia/Source/ImgMediaEditor/Private/ImgMediaEditorModule.cpp
--- a/Engine/Plugins/Media/ImgMedia/Source/ImgMediaEditor/Private/ImgMediaEditorModule.cpp
+++ b/Engine/Plugins/Media/ImgMedia/Source/ImgMediaEditor/Private/ImgMediaEditorModule.cpp
@@ -217,9 +217,9 @@ protected:
 				// Add menu entry.
 				if (UImgMediaSource* ImgMediaSource = Cast<UImgMediaSource>(Context->SelectedAssets[0].GetAsset()))
 				{
-					// Get path from media source.
+					// Get the full path from media source.
 					// Make sure it ends with / so it opens the directory.
-					FString InputPath = ImgMediaSource->GetSequencePath();
+					FString InputPath = ImgMediaSource->GetFullPath();
 					if (InputPath.EndsWith(TEXT("/")) == false)
 					{
 						InputPath += TEXT("/");
