[ImgMedia] ProcessEXR: Fix the toaster messages and Fix crash if closing the panel while processing.

- The number of image processed was based of "num done" counter that started at 0 for the first image, which was misleading. If we think about it as an index, it is more obvious that NumDone = Index + 1 for the message. Renamed the variable to make it more obvious.
- Added more input/output path validation along with corresponding log/notification messages. That will lead to more informative/accurate messages to the user as to what is wrong in the setup.
- Added some thread safety to avoid the widget from being deleted (by closing the panel) while async tasks are on going. The "NumAsyncTasks" counter is used for that, see ~SImgMediaProcessEXR() which will now cancel and wait for tasks to finish to avoid dangling "this" pointer on the threads.

#Synced-CL 49891360

diff --git a/Engine/Plugins/Media/ImgMedia/Source/ImgMediaEditor/Private/Widgets/SImgMediaProcessEXR.cpp b/Engine/Plugins/Media/ImgMedia/Source/ImgMediaEditor/Private/Widgets/SImgMediaProcessEXR.cpp
--- a/Engine/Plugins/Media/ImgMedia/Source/ImgMediaEditor/Private/Widgets/SImgMediaProcessEXR.cpp
+++ b/Engine/Plugins/Media/ImgMedia/Source/ImgMediaEditor/Private/Widgets/SImgMediaProcessEXR.cpp
@@ -42,9 +42,63 @@
 
 #define LOCTEXT_NAMESPACE "ImgMediaProcessImages"
 
+namespace UE::ImgMediaEditor::ProcessEXR
+{
+	/**
+	 * Helper function to set the Notification text from a worker thread.
+	 * @param InNotificationWeak Weak reference to the Notification object.
+	 * @param InText Text to set.
+	 */
+	void SetNotificationTextSafe(const TWeakPtr<SNotificationItem>& InNotificationWeak, const FText& InText)
+	{
+		Async(EAsyncExecution::TaskGraphMainThread, [NotificationWeak = InNotificationWeak, Text = InText]()
+		{
+			if (TSharedPtr<SNotificationItem> Notification = NotificationWeak.Pin())
+			{
+				Notification->SetText(Text);
+			}
+		});
+	}
+
+	/**
+	 * Helper function to set the Notification text from a worker thread.
+	 * @param InNotification Reference to the Notification object.
+	 * @param InText Text to set.
+	 */
+	void SetNotificationTextSafe(const TSharedPtr<SNotificationItem>& InNotification, const FText& InText)
+	{
+		if (InNotification.IsValid())
+		{
+			SetNotificationTextSafe(InNotification.ToWeakPtr(), InText);
+		}
+	}
+}
+
 SImgMediaProcessEXR::~SImgMediaProcessEXR()
 {
 	CleanUp();
+
+	// Check if we have ongoing async threads preventing us from deleting the widget.
+	if (NumAsyncTasks.load() > 0)
+	{
+		// Signal tasks to stop.
+		bIsCancelling = true;
+
+		// Wait for all async tasks to complete.
+		while (NumAsyncTasks.load() > 0)
+		{
+			FPlatformProcess::Sleep(0.2f);
+		}
+
+		// Close the toaster message.
+		if (ConfirmNotification.IsValid())
+		{
+			ConfirmNotification->SetEnabled(false);
+			ConfirmNotification->SetCompletionState(SNotificationItem::CS_Fail);
+			ConfirmNotification->ExpireAndFadeout();
+			ConfirmNotification.Reset();
+		}
+	}
 }
 
 BEGIN_SLATE_FUNCTION_BUILD_OPTIMIZATION
@@ -190,10 +244,46 @@ FReply SImgMediaProcessEXR::OnProcessImagesClicked()
 		}
 		else
 		{
+			TWeakPtr<SImgMediaProcessEXR> WeakSelf = TSharedRef<SImgMediaProcessEXR>(AsShared(), this).ToWeakPtr();
+
 			// Start async task to process files.
-			Async(EAsyncExecution::Thread, [this]()
+			Async(EAsyncExecution::Thread, [WeakSelf]()
 			{
-				ProcessAllImages();
+				// Self is only pinned temporarily for thread safety, then released to allow the job to be cancelled (by calling destructor)
+				// if the panel is closed. For the remainer of the scope, we protect the validity of the Self pointer with NumAsyncTasks
+				// counter instead (see destructor).
+				SImgMediaProcessEXR* SelfPtr;
+				{
+					const TSharedPtr<SImgMediaProcessEXR> Self = WeakSelf.Pin();
+					if (!Self)
+					{
+						return;
+					}
+					++Self->NumAsyncTasks;	// This will protect from destruction while async task is going.
+					SelfPtr = Self.Get();	// Keep raw pointer for the remainer of the scope.
+				}
+
+				SelfPtr->ProcessAllImages();
+				--SelfPtr->NumAsyncTasks;
+
+				// Close notification. Must be run on the main thread.
+				Async(EAsyncExecution::TaskGraphMainThread, [WeakSelf]()
+				{
+					if (TSharedPtr<SImgMediaProcessEXR> Self = WeakSelf.Pin())
+					{
+						if (Self->ConfirmNotification.IsValid())
+						{
+							Self->ConfirmNotification->SetEnabled(false);
+							Self->ConfirmNotification->SetCompletionState((Self->bIsCancelling) ? SNotificationItem::CS_Fail : SNotificationItem::CS_Success);
+							Self->ConfirmNotification->ExpireAndFadeout();
+						}
+
+						// Done with processing.
+						Self->bIsProcessing = false;
+						Self->bIsCancelling = false;
+						Self->UpdateWidgets();
+					}
+				});
 			});
 		}
 	}
@@ -214,6 +304,7 @@ FReply SImgMediaProcessEXR::OnCancelClicked()
 
 void SImgMediaProcessEXR::ProcessAllImages()
 {
+	using namespace UE::ImgMediaEditor::ProcessEXR;
 	TRACE_CPUPROFILER_EVENT_SCOPE(SImgMediaProcessEXR::ProcessAllImages);
 
 	int32 InTileWidth = Options->bEnableTiling ? Options->TileSizeX : 0;
@@ -224,10 +315,38 @@ void SImgMediaProcessEXR::ProcessAllImages()
 	// Create output directory.
 	FString OutPath = Options->OutputPath.Path;
 	IPlatformFile& PlatformFile = FPlatformFileManager::Get().GetPlatformFile();
-	PlatformFile.CreateDirectoryTree(*OutPath);
+	if (!FPaths::DirectoryExists(OutPath))
+	{
+		if (!PlatformFile.CreateDirectoryTree(*OutPath))
+		{
+			const FText Message = FText::Format(LOCTEXT("FailedToCreateOutputTree", "Failed to create output directory \"{0}\""), FText::FromString(OutPath)); 
+			UE_LOG(LogImgMediaEditor, Error, TEXT("ProcessEXR: %s."), *Message.ToString());
+			SetNotificationTextSafe(ConfirmNotification, Message);
+			bIsCancelling = true;
+			return;
+		}
+	}
 
 	// Get source files.
-	FString SequencePath = FPaths::ConvertRelativePathToFull(FPaths::GetPath(Options->InputPath.FilePath));
+	const FString InputPath = FPaths::GetPath(Options->InputPath.FilePath);
+	if (InputPath.IsEmpty())
+	{
+		const FText Message = FText::Format(LOCTEXT("InvalidInputPath", "Specified input path \"{0}\" is invalid."), FText::FromString(Options->InputPath.FilePath)); 
+		UE_LOG(LogImgMediaEditor, Error, TEXT("ProcessEXR: %s."), *Message.ToString());
+		SetNotificationTextSafe(ConfirmNotification, Message);
+		bIsCancelling = true;
+		return;
+	}
+
+	const FString SequencePath = FPaths::ConvertRelativePathToFull(InputPath);
+	if (!FPaths::DirectoryExists(SequencePath))
+	{
+		const FText Message = FText::Format(LOCTEXT("InvalidSequencePath", "Specified input directory directory \"{0}\" does not exist."), FText::FromString(SequencePath));
+		UE_LOG(LogImgMediaEditor, Error, TEXT("ProcessEXR: %s."), *Message.ToString());
+		SetNotificationTextSafe(ConfirmNotification, Message);
+		bIsCancelling = true;
+		return;
+	}
 
 	TArray<FString> FoundFiles;
 	IFileManager::Get().FindFiles(FoundFiles, *SequencePath, TEXT("*"));
@@ -236,122 +355,104 @@ void SImgMediaProcessEXR::ProcessAllImages()
 	UE_LOG(LogImgMediaEditor, Warning, TEXT("Found %i image files in %s to import."), FoundFiles.Num(), *SequencePath);
 	if (FoundFiles.Num() == 0)
 	{
-		UE_LOG(LogImgMediaEditor, Error, TEXT("No files to import."));
+		const FText Message = LOCTEXT("NoFilesToImport", "No files to import.");
+		UE_LOG(LogImgMediaEditor, Error, TEXT("ProcessEXR: %s."), *Message.ToString());
+		SetNotificationTextSafe(ConfirmNotification, Message);
+		bIsCancelling = true;
+		return;
+	}
+	
+	// Create image wrapper
+	FString Ext = FPaths::GetExtension(FoundFiles[0]);
+	EImageFormat ImageFormat = ImageWrapperHelper::GetImageFormat(Ext);
+
+	if (ImageFormat == EImageFormat::Invalid)
+	{
+		const FText Message = FText::Format(LOCTEXT("InvalidFileFormat", "Invalid file format \"{0}\""), FText::FromString(Ext));
+		UE_LOG(LogImgMediaEditor, Error, TEXT("ProcessEXR: %s."), *Message.ToString());
+		SetNotificationTextSafe(ConfirmNotification, Message);
+		bIsCancelling = true;
+		return;
 	}
-	else
+	
+	// Get number of threads to use.
+	int32 NumThreads = Options->NumThreads;
+	if (NumThreads <= 0)
 	{
-		// Create image wrapper
-		FString Ext = FPaths::GetExtension(FoundFiles[0]);
-		EImageFormat ImageFormat = ImageWrapperHelper::GetImageFormat(Ext);
+		NumThreads = 8;
+	}
+
+	// Loop through all files.
+	int FileIndex = 0;
+	int TotalNum = FoundFiles.Num();
+	std::atomic<int32> NumActive = 0;
 
-		if (ImageFormat == EImageFormat::Invalid)
+	for (const FString& FileName : FoundFiles)
+	{
+		// Wait for threads to finish if we have to many.
+		while (NumActive >= NumThreads)
 		{
-			UE_LOG(LogImgMediaEditor, Error, TEXT("Invalid file format %s"), *Ext);
+			FPlatformProcess::Sleep(0.1f);
 		}
-		else
-		{			
-			// Get number of threads to use.
-			int32 NumThreads = Options->NumThreads;
-			if (NumThreads <= 0)
-			{
-				NumThreads = 8;
-			}
+		++NumActive;
+
+		// Update notification with current status.
+		SetNotificationTextSafe(ConfirmNotification,
+			FText::Format(LOCTEXT("ImgMediaCompleted", "ImgMedia Completed {0}/{1}"), FText::AsNumber(FileIndex+1), FText::AsNumber(TotalNum)));
+
+		Async(EAsyncExecution::Thread, [this, SequencePath, FileName, OutPath,
+			Ext, &NumActive, InTileWidth, InTileHeight,
+			TileBorder, bEnableMips]() mutable
+		{
+			FString FullFileName = FPaths::Combine(SequencePath, FileName);
+			FImage Image;
 
-			// Loop through all files.
-			int NumDone = 0;
-			int TotalNum = FoundFiles.Num();
-			std::atomic<int32> NumActive = 0;
-			TSharedPtr<SNotificationItem> LocalConfirmNotification = ConfirmNotification;
-			for (const FString& FileName : FoundFiles)
+			// Load image into buffer.
 			{
-				// Wait for threads to finish if we have to many.
-				while (NumActive >= NumThreads)
+				TRACE_CPUPROFILER_EVENT_SCOPE(SImgMediaProcessEXR::ProcessImageCustom:LoadImage);
+				if (!FImageUtils::LoadImage(*FullFileName, Image))
 				{
-					FPlatformProcess::Sleep(0.1f);
+					UE_LOG(LogImgMediaEditor, Error, TEXT("Failed to load %s"), *FullFileName);
+					--NumActive;
+					return;
 				}
-				NumActive++;
-				
-				// Update notification with current status.
-				Async(EAsyncExecution::TaskGraphMainThread, [LocalConfirmNotification, NumDone, TotalNum]()
-				{
-					if (LocalConfirmNotification.IsValid())
-					{
-						LocalConfirmNotification->SetText(
-							FText::Format(LOCTEXT("ImgMediaCompleted", "ImgMedia Completed {0}/{1}"),
-								FText::AsNumber(NumDone), FText::AsNumber(TotalNum)));
-					}
-				});
-				NumDone++;
-
-				Async(EAsyncExecution::Thread, [this, SequencePath, FileName, OutPath,
-					Ext, &NumActive, InTileWidth, InTileHeight,
-					TileBorder, bEnableMips]() mutable
-				{
-					FString FullFileName = FPaths::Combine(SequencePath, FileName);
-					FImage Image;
-
-					// Load image into buffer.
-					{
-						TRACE_CPUPROFILER_EVENT_SCOPE(SImgMediaProcessEXR::ProcessImageCustom:LoadImage);
-						if (!FImageUtils::LoadImage(*FullFileName, Image))
-						{
-							UE_LOG(LogImgMediaEditor, Error, TEXT("Failed to load %s"), *FullFileName);
-							NumActive--;
-							return;
-						}
-					}
+			}
 
-					// Convert to linear float16 if not already.
-					{
-						TRACE_CPUPROFILER_EVENT_SCOPE(SImgMediaProcessEXR::ProcessImageCustom:ChangeFormat);
-						Image.ChangeFormat(ERawImageFormat::RGBA16F, EGammaSpace::Linear);
-					}
+			// Convert to linear float16 if not already.
+			{
+				TRACE_CPUPROFILER_EVENT_SCOPE(SImgMediaProcessEXR::ProcessImageCustom:ChangeFormat);
+				Image.ChangeFormat(ERawImageFormat::RGBA16F, EGammaSpace::Linear);
+			}
 
-					FImageParameters ImageParameters;
-					GetImageParameters(Ext, FullFileName, ImageParameters);
+			FImageParameters ImageParameters;
+			GetImageParameters(Ext, FullFileName, ImageParameters);
 
-					FString Name = FPaths::Combine(OutPath, FileName);
-					if (Ext != TEXT("exr"))
-					{
-						Name = FPaths::ChangeExtension(Name, TEXT("exr"));
-					}
+			FString Name = FPaths::Combine(OutPath, FileName);
+			if (Ext != TEXT("exr"))
+			{
+				Name = FPaths::ChangeExtension(Name, TEXT("exr"));
+			}
 
-					ProcessImageCustom(Image, InTileWidth, InTileHeight, TileBorder,
-						bEnableMips, ImageParameters, Name);
+			ProcessImageCustom(Image, InTileWidth, InTileHeight, TileBorder,
+				bEnableMips, ImageParameters, Name);
 
-					NumActive--;
-				});
+			--NumActive;
+		});
 
-				// Do we want to cancel?
-				if (bIsCancelling)
-				{
-					break;
-				}
-			}
+		++FileIndex;
 
-			// Wait for all our tasks to finish.
-			while (NumActive > 0)
-			{
-				FPlatformProcess::Sleep(0.2f);
-			}
+		// Do we want to cancel?
+		if (bIsCancelling)
+		{
+			break;
 		}
 	}
 
-	// Close notification. Must be run on the main thread.
-	Async(EAsyncExecution::TaskGraphMainThread, [this]()
+	// Wait for all our tasks to finish.
+	while (NumActive > 0)
 	{
-		if (ConfirmNotification.IsValid())
-		{
-			ConfirmNotification->SetEnabled(false);
-			ConfirmNotification->SetCompletionState(bIsCancelling ? SNotificationItem::CS_Fail : SNotificationItem::CS_Success);
-			ConfirmNotification->ExpireAndFadeout();
-		}
-
-		// Done with processing.
-		bIsProcessing = false;
-		bIsCancelling = false;
-		UpdateWidgets();
-	});
+		FPlatformProcess::Sleep(0.2f);
+	}
 }
 
 void SImgMediaProcessEXR::GetImageParameters(const FString& Ext, const FString& File,
@@ -858,9 +959,11 @@ void SImgMediaProcessEXR::HandleProcessing()
 						Async(EAsyncExecution::Thread, [this, RawData = MoveTemp(Image.RawData), Width, Height,
 							InTileWidth, InTileHeight, TileBorder, bEnableMips, Name]() mutable
 						{
+							++NumAsyncTasks;
 							ProcessImageCustomRawData(RawData, Width, Height,
 								InTileWidth, InTileHeight, TileBorder, bEnableMips,
 								GlobalImageParameters, Name);
+							--NumAsyncTasks;
 						});
 					}
 					else
diff --git a/Engine/Plugins/Media/ImgMedia/Source/ImgMediaEditor/Private/Widgets/SImgMediaProcessEXR.h b/Engine/Plugins/Media/ImgMedia/Source/ImgMediaEditor/Private/Widgets/SImgMediaProcessEXR.h
--- a/Engine/Plugins/Media/ImgMedia/Source/ImgMediaEditor/Private/Widgets/SImgMediaProcessEXR.h
+++ b/Engine/Plugins/Media/ImgMedia/Source/ImgMediaEditor/Private/Widgets/SImgMediaProcessEXR.h
@@ -211,7 +211,7 @@ private:
 	FTimespan CurrentTime;
 	/** How long a frame is. */
 	FTimespan FrameDuration;
-	/** Current frame index we are processinig. */
+	/** Current frame index we are processing. */
 	int32 CurrentFrameIndex = 0;
 
 	/** True if we are currently processing. */
@@ -220,5 +220,10 @@ private:
 	bool bIsCancelling = false;
 	/** True if we are using a player to read in the images. */
 	bool bUsePlayer = false;
-	
+
+	/**
+	 * Counter for the number of "async tasks" that are currently running.
+	 * This used by the main thread to protect against the widget being deleted while in use by the async tasks.
+	 */
+	std::atomic<int32> NumAsyncTasks = 0;
 };
