ImgMedia: Warning to allow future deprecation of mip level reading from multiple folders.

The preferred paths moving forward will be TMV, EXR and DDS.

#Synced-CL 50746286

diff --git a/Engine/Plugins/Media/ImgMedia/Source/ImgMedia/Private/Loader/ImgMediaLoader.cpp b/Engine/Plugins/Media/ImgMedia/Source/ImgMedia/Private/Loader/ImgMediaLoader.cpp
--- a/Engine/Plugins/Media/ImgMedia/Source/ImgMedia/Private/Loader/ImgMediaLoader.cpp
+++ b/Engine/Plugins/Media/ImgMedia/Source/ImgMedia/Private/Loader/ImgMediaLoader.cpp
@@ -1173,6 +1173,11 @@ void FImgMediaLoader::FindMips(const FString& SequencePath)
 			ImagePaths.Emplace(MipFiles);
 		}
 	}
+
+	if (ImagePaths.Num() > 1)
+	{
+		UE_CALL_ONCE([&] { UE_LOG(LogImgMedia, Warning, TEXT("Loading mips from separate folders will be deprecated soon (%s). Please prefer using TMV, EXR or DDS with mip levels."), *SequencePath) });
+	}
 }
 
 
