[ImgMedia] Prune Stale Tracker objects from the mip map info.

#Synced-CL 51094493

diff --git a/Engine/Plugins/Media/ImgMedia/Source/ImgMedia/Internal/Assets/ImgMediaMipMapInfo.h b/Engine/Plugins/Media/ImgMedia/Source/ImgMedia/Internal/Assets/ImgMediaMipMapInfo.h
--- a/Engine/Plugins/Media/ImgMedia/Source/ImgMedia/Internal/Assets/ImgMediaMipMapInfo.h
+++ b/Engine/Plugins/Media/ImgMedia/Source/ImgMedia/Internal/Assets/ImgMediaMipMapInfo.h
@@ -217,6 +217,12 @@ public:
 	/** Get the target view resolution. */
 	EMediaTextureTargetViewResolution GetTargetViewResolutionMask() const;
 
+	/** Returns true if this object is stale. */
+	bool IsStale() const
+	{
+		return MeshComponent.IsStale() || !Tracker.IsValid();
+	}
+	
 protected:
 
 	/** Mesh component onto which the media is displayed. */
@@ -268,6 +274,11 @@ public:
 	 */
 	void ClearAllObjects();
 
+	/**
+	 * Remove all objects with stale owners.
+	 */
+	void PruneStaleObjects();
+
 	/**
 	 * Provide information on the texture needed for our image sequence.
 	 *
diff --git a/Engine/Plugins/Media/ImgMedia/Source/ImgMedia/Private/Assets/ImgMediaMipMapInfo.cpp b/Engine/Plugins/Media/ImgMedia/Source/ImgMedia/Private/Assets/ImgMediaMipMapInfo.cpp
--- a/Engine/Plugins/Media/ImgMedia/Source/ImgMedia/Private/Assets/ImgMediaMipMapInfo.cpp
+++ b/Engine/Plugins/Media/ImgMedia/Source/ImgMedia/Private/Assets/ImgMediaMipMapInfo.cpp
@@ -1028,6 +1028,20 @@ void FImgMediaMipMapInfo::ClearAllObjects()
 	UnsubscribeEndFrame();
 }
 
+void FImgMediaMipMapInfo::PruneStaleObjects()
+{
+	FScopeLock Lock(&ObjectsCriticalSection);
+
+	for (TArray<FImgMediaMipMapObjectInfo*>::TIterator It(Objects); It; ++It)
+	{
+		FImgMediaMipMapObjectInfo* Info = (*It);
+		if (Info && Info->IsStale())
+		{
+			delete Info;
+			It.RemoveCurrent();
+		}
+	}
+}
 
 void FImgMediaMipMapInfo::SetTextureInfo(FName InSequenceName, int32 InNumMipMaps, const FIntPoint& InSequenceDim,
 	const FMediaTextureTilingDescription& InTilingDescription)
diff --git a/Engine/Plugins/Media/ImgMedia/Source/ImgMedia/Private/Player/ImgMediaPlayer.cpp b/Engine/Plugins/Media/ImgMedia/Source/ImgMedia/Private/Player/ImgMediaPlayer.cpp
--- a/Engine/Plugins/Media/ImgMedia/Source/ImgMedia/Private/Player/ImgMediaPlayer.cpp
+++ b/Engine/Plugins/Media/ImgMedia/Source/ImgMedia/Private/Player/ImgMediaPlayer.cpp
@@ -274,6 +274,8 @@ bool FImgMediaPlayer::Open(const FString& Url, const IMediaOptions* Options, con
 			MipMapInfo = StaticCastSharedPtr<FImgMediaMipMapInfo, IMediaOptions::FDataContainer, ESPMode::ThreadSafe>(DataContainer);
 			if (MipMapInfo.IsValid())
 			{
+				MipMapInfo->PruneStaleObjects();
+
 				if (ViewTexture)
 				{
 					MediaTextureWeakPtr = ViewTexture;
