[ImgMedia] DDS Reader: adding support for full mip chain when block compression is not used.

Also, fixing incorrect sRGB flag for float format.

#Synced-CL 51094493

diff --git a/Engine/Plugins/Media/ImgMedia/Source/ImgMedia/Private/Readers/DdsImgMediaReader.cpp b/Engine/Plugins/Media/ImgMedia/Source/ImgMedia/Private/Readers/DdsImgMediaReader.cpp
--- a/Engine/Plugins/Media/ImgMedia/Source/ImgMedia/Private/Readers/DdsImgMediaReader.cpp
+++ b/Engine/Plugins/Media/ImgMedia/Source/ImgMedia/Private/Readers/DdsImgMediaReader.cpp
@@ -15,6 +15,52 @@ namespace ImgMedia
 {
 	namespace DDS
 	{
+		/** Descriptor for a DDS pixel format. */
+		struct FDDSFormatInfo
+		{
+			UE::DDS::EDXGIFormat Format;
+			uint32 Id;
+			uint8 BytesPerUnit;
+			bool bIsBlockCompressed;
+		};
+
+		static const FDDSFormatInfo FormatInfoList[] = 
+		{
+			#define RGBFMT(Name, Id, BytesPerUnit) { UE::DDS::EDXGIFormat::Name, Id, BytesPerUnit, false },
+			#define BCNFMT(Name, Id, BytesPerUnit) { UE::DDS::EDXGIFormat::Name, Id, BytesPerUnit, true },
+			#define ODDFMT(Name, Id) { UE::DDS::EDXGIFormat::Name, Id, 0, false },
+			UE_DXGI_FORMAT_LIST
+			#undef RGBFMT
+			#undef BCNFMT
+			#undef ODDFMT
+		};
+
+		static const FDDSFormatInfo* GetDDSFormatInfo(UE::DDS::EDXGIFormat InFormat)
+		{
+			// need to handle this special because UNKNOWN _does_ appear in the FormatInfoList,
+			// but we don't want to treat it as legal
+			if (InFormat == UE::DDS::EDXGIFormat::UNKNOWN)
+			{
+				return nullptr;
+			}
+
+			for (size_t Index = 0; Index < sizeof(FormatInfoList)/sizeof(*FormatInfoList); ++Index)
+			{
+				if (InFormat == FormatInfoList[Index].Format) 
+				{
+					return &FormatInfoList[Index];
+				}
+			}
+
+			return nullptr;
+		}
+
+		bool IsFormatBlockCompressed(UE::DDS::EDXGIFormat InDXGIFormat)
+		{
+			const FDDSFormatInfo* FormatInfo = GetDDSFormatInfo(InDXGIFormat);
+			return FormatInfo ? FormatInfo->bIsBlockCompressed : false;
+		}
+		
 		/* Returns the media texture sample format from a DXGI format type. */
 		bool GetMediaTextureSampleFormat(const UE::DDS::EDXGIFormat& InDXGIFormat, EMediaTextureSampleFormat& OutFormat, bool& bOutIsBlockCompressed)
 		{
@@ -66,22 +112,29 @@ namespace ImgMedia
 			return true;
 		}
 
-		/* Only consider mips larger than the block size as valid. */
-		inline bool IsMipValid(const UE::DDS::FDDSMip& InMip)
+		/*
+		 * For block compressed formats, only consider mips larger than the block size as valid. 
+		 */
+		inline bool IsMipValid(const UE::DDS::FDDSMip& InMip, bool bIsBlockCompressed)
 		{
-			// See FD3D12Texture::UpdateTexture2D block size checks, and also
-			// TODO: D3D12_FEATURE_DATA_D3D12_OPTIONS8.UnalignedBlockTexturesSupported
+			if (bIsBlockCompressed)
+			{
+				// See FD3D12Texture::UpdateTexture2D block size checks, and also
+				// TODO: D3D12_FEATURE_DATA_D3D12_OPTIONS8.UnalignedBlockTexturesSupported
+				return FMath::Min(InMip.Width, InMip.Height) >= 4 /* CompressionBlockSize */;
+			}
 
-			return FMath::Min(InMip.Width, InMip.Height) >= 4 /* CompressionBlockSize */;
+			return true;
 		}
 
 		/* Returns the DDS total mip data size. */
 		SIZE_T GetTotalMipDataSize(const TUniquePtr<UE::DDS::FDDSFile>& DDS)
 		{
+			const bool bIsBlockCompressed = IsFormatBlockCompressed(DDS->DXGIFormat);
 			SIZE_T TotalMipDataSize = 0;
 			for (const UE::DDS::FDDSMip& Mip : DDS->Mips)
 			{
-				if (!IsMipValid(Mip))
+				if (!IsMipValid(Mip, bIsBlockCompressed))
 				{
 					break;
 				}
@@ -95,12 +148,13 @@ namespace ImgMedia
 		/* Returns the DDS file mip count, accounting for engine maximum. */
 		int32 GetMipCount(const TUniquePtr<UE::DDS::FDDSFile>& DDS)
 		{
+			const bool bIsBlockCompressed = IsFormatBlockCompressed(DDS->DXGIFormat);
 			int32 MipCount = 0;
 
 			// Mips are ordered starting from mip 0 (full-size texture) decreasing in size
 			for (const UE::DDS::FDDSMip& Mip : DDS->Mips)
 			{
-				if (!IsMipValid(Mip))
+				if (!IsMipValid(Mip, bIsBlockCompressed))
 				{
 					break;
 				}
@@ -111,9 +165,35 @@ namespace ImgMedia
 			return FMath::Min(MipCount, static_cast<int32>(MAX_TEXTURE_MIP_COUNT));
 		}
 
+		/** Returns true if the given format is a floating point type. */
+		bool IsFloatFormat(UE::DDS::EDXGIFormat InFormat)
+		{
+			switch (InFormat)
+			{
+			case UE::DDS::EDXGIFormat::R32G32B32A32_FLOAT:
+			case UE::DDS::EDXGIFormat::R32G32B32_FLOAT:
+			case UE::DDS::EDXGIFormat::R16G16B16A16_FLOAT:
+			case UE::DDS::EDXGIFormat::R32G32_FLOAT:
+			case UE::DDS::EDXGIFormat::R11G11B10_FLOAT:
+			case UE::DDS::EDXGIFormat::R16G16_FLOAT:
+			case UE::DDS::EDXGIFormat::D32_FLOAT:
+			case UE::DDS::EDXGIFormat::R32_FLOAT:
+			case UE::DDS::EDXGIFormat::R16_FLOAT:
+				return true;
+			default:
+				return false;
+			}
+		}
+
 		/* Returns whether the DDS file is sRGB-encoded. */
 		bool DDSPayloadIsSRGB(const TUniquePtr<UE::DDS::FDDSFile>& DDS)
 		{
+			// Float formats are considered linear.
+			if (IsFloatFormat(DDS->DXGIFormat))
+			{
+				return false;
+			}
+			
 			if (DDS->CreateFlags & UE::DDS::FDDSFile::CREATE_FLAG_WAS_D3D9)
 			{
 				// no SRGB info in Dx9 format
