[ImgMedia] Adding a CVar "ImgMedia.EnableMipMapTileDebugDraw" for mipmap info debug spheres to help trouble shoot tile update issues.

This works in game mode too ti test on nDisplay nodes.

#Synced-CL 51094493

diff --git a/Engine/Plugins/Media/ImgMedia/Source/ImgMedia/Private/Assets/ImgMediaMipMapInfo.cpp b/Engine/Plugins/Media/ImgMedia/Source/ImgMedia/Private/Assets/ImgMediaMipMapInfo.cpp
--- a/Engine/Plugins/Media/ImgMedia/Source/ImgMedia/Private/Assets/ImgMediaMipMapInfo.cpp
+++ b/Engine/Plugins/Media/ImgMedia/Source/ImgMedia/Private/Assets/ImgMediaMipMapInfo.cpp
@@ -37,6 +37,14 @@ static TAutoConsoleVariable<float> CVarImgMediaMipLevelPadding(
 	TEXT("Value padded onto the estimated (minimum and maximum) mipmap levels used by the loader.\n"),
 	ECVF_Default);
 
+static TAutoConsoleVariable<bool> CVarImgMediaEnableMipMapTileDebugDraw(
+	TEXT("ImgMedia.EnableMipMapTileDebugDraw"),
+	false,
+	TEXT("Enables drawing of the bounding volume around visible tiles for trouble shooting tile update issues.")
+	TEXT("   0: off (default)\n")
+	TEXT("   1: on\n"),
+	ECVF_Default);
+
 FImgMediaTileSelection::FImgMediaTileSelection(int32 NumTilesX, int32 NumTilesY, bool bDefaultVisibility)
 	: Tiles(bDefaultVisibility, NumTilesX * NumTilesY)
 	, Dimensions(NumTilesX, NumTilesY)
@@ -423,7 +431,9 @@ namespace {
 
 			const float MipMapBias = ObjectInfo->MipMapLODBias;
 			const float MipMapLevelPadding = FMath::Max(CVarImgMediaMipLevelPadding.GetValueOnAnyThread(), 0.0f);
-
+#if WITH_EDITOR
+			const bool bDrawDebugSpheres = CVarImgMediaEnableMipMapTileDebugDraw.GetValueOnAnyThread();
+#endif
 			const FIntPoint& SequenceTileNum = InSequenceInfo.TilingDescription.TileNum;
 
 			// To avoid calculating tile corner mip levels multiple times over, we cache them in this array.
@@ -581,16 +591,18 @@ namespace {
 
 							VisibleTiles[CurrentMipLevel].SetVisible(Tile.X, Tile.Y);
 						}
-#if false
 #if WITH_EDITOR
-						// Enable this to draw a sphere where each tile is.
-						Async(EAsyncExecution::TaskGraphMainThread, [TileCenterWS, TileRadiusInWS]()
+						if(bDrawDebugSpheres)
+						{
+							Async(EAsyncExecution::TaskGraphMainThread, [TileCenterWS, TileRadiusInWS]()
 							{
-								UWorld* World = GEditor->GetEditorWorldContext().World();
-								DrawDebugSphere(World, TileCenterWS, TileRadiusInWS, 8, FColor::Red, false, 0.05f);
+								if (const UWorld* World = GEditor ? GEditor->GetEditorWorldContext().World() : GWorld)
+								{
+									DrawDebugSphere(World, TileCenterWS, TileRadiusInWS, 8, FColor::Red, false, 0.05f);
+								}
 							});
+						}
 #endif // WITH_EDITOR
-#endif // false
 					}
 				}
 			}
@@ -669,6 +681,9 @@ namespace {
 				return;
 			}
 
+#if WITH_EDITOR
+			const bool bDrawDebugSpheres = CVarImgMediaEnableMipMapTileDebugDraw.GetValueOnAnyThread();
+#endif
 			const FVector2f MeshRange = FVector2f(ObjectInfo->MeshRange);
 			const float MipMapBias = ObjectInfo->MipMapLODBias;
 			const float MipMapLevelPadding = FMath::Max(CVarImgMediaMipLevelPadding.GetValueOnAnyThread(), 0.0f);
@@ -832,17 +847,19 @@ namespace {
 									VisibleTiles[Level].SetVisible(TileX / MipLevelDiv, TileY / MipLevelDiv);
 								}
 							}
-#if false
 #if WITH_EDITOR
-							float DebugSphereRadius = FMath::Max(CollisionSphereRadius, MeshTransform.GetMaximumAxisScale());
-							// Enable this to draw a sphere where each tile is.
-							Async(EAsyncExecution::TaskGraphMainThread, [TileLocationWS, DebugSphereRadius]()
+							if(bDrawDebugSpheres)
+							{
+								const float DebugSphereRadius = FMath::Max(CollisionSphereRadius, MeshTransform.GetMaximumAxisScale());
+								Async(EAsyncExecution::TaskGraphMainThread, [TileLocationWS, DebugSphereRadius]()
 								{
-									UWorld* World = GEditor->GetEditorWorldContext().World();
-									DrawDebugSphere(World, TileLocationWS, DebugSphereRadius, 8, FColor::Red, false, 0.05f);
+									if (const UWorld* World = GEditor ? GEditor->GetEditorWorldContext().World() : GWorld)
+									{
+										DrawDebugSphere(World, TileLocationWS, DebugSphereRadius, 8, FColor::Red, false, 0.05f);
+									}
 								});
+							}
 #endif // WITH_EDITOR
-#endif // false
 						}
 					}
 				};
