[ImgMedia] Fix Sphere object tile calculation in split frustum nDisplay configuration

#Synced-CL 51094493

diff --git a/Engine/Plugins/Media/ImgMedia/Source/ImgMedia/Private/Assets/ImgMediaMipMapInfo.cpp b/Engine/Plugins/Media/ImgMedia/Source/ImgMedia/Private/Assets/ImgMediaMipMapInfo.cpp
--- a/Engine/Plugins/Media/ImgMedia/Source/ImgMedia/Private/Assets/ImgMediaMipMapInfo.cpp
+++ b/Engine/Plugins/Media/ImgMedia/Source/ImgMedia/Private/Assets/ImgMediaMipMapInfo.cpp
@@ -690,7 +690,7 @@ namespace {
 			const FTransform MeshTransform = Mesh->GetComponentTransform();
 			const int32 MaxLevel = InSequenceInfo.NumMipLevels - 1;
 
-			// Include all tiles containted in the visible UV region
+			// Include all tiles contained in the visible UV region
 			const FIntPoint& SequenceTileNum = InSequenceInfo.TilingDescription.TileNum;
 			const FVector2f SequencePartialTileNum = InSequenceInfo.GetPartialTileNum();
 
@@ -746,8 +746,8 @@ namespace {
 				const float TileReductionFactor = 0.5f;
 				if (bAdaptivePoleMipUpscaling)
 				{
-					NumOfTilesToLoadAtThePole = FMath::CeilToInt32(SequencePartialTileNum.X * ((float)SequencePartialTileNum.Y * PoleEdgeDegrees / 180.f))* TileReductionFactor;
-					PoleMipBias = FMath::CeilToInt32(FMath::LogX(4., TotalNumTiles/ NumOfTilesToLoadAtThePole));
+					NumOfTilesToLoadAtThePole = FMath::CeilToInt32(SequencePartialTileNum.X * ((float)SequencePartialTileNum.Y * PoleEdgeDegrees / 180.f)) * TileReductionFactor;
+					PoleMipBias = (NumOfTilesToLoadAtThePole > 0) ? FMath::CeilToInt32(FMath::LogX(4., TotalNumTiles/ NumOfTilesToLoadAtThePole)) : 0;
 				}
 
 				// Everything above this value gets upscaled.
@@ -779,7 +779,7 @@ namespace {
 						if (ViewUV.ComponentwiseAllGreaterOrEqual(TileMinCornerUV) && ViewUV.ComponentwiseAllLessThan(TileMaxCornerUV))
 						{
 							TileUV = ViewUV;
-							CollisionSphereRadius = 0.0f;
+							CollisionSphereRadius *= 0.5f;
 						}
 
 						// The resulting location used for tile mip level estimation (either directly at the view center or its closest tile corner).
