EXR Multi-Layer Support for ImgMedia with the EXR CPU reader.

* Adds LayerName to UImgMediaSource to select a named layer from multi-layer EXR sequences, with exact, wildcard, and grayscale single-channel support.
* LayerName defaults to empty, which is a no-op that preserves the original read behavior (default un-prefixed channels).
* CPU reader: replaces FRgbaInputFile with a new FExrInputFile that reads named channels via Imf::InputFile, with layer resolution cached on first read. FExrInputFile is a new class rather than an extension of FRgbaInputFile because Imf::InputFile and Imf::RgbaInputFile are unrelated OpenEXR types with different channel-selection APIs.
* GPU reader: forwards LayerName to FExrImgMediaReaderGpu, using the existing GPU path for pure-layer files and falling back to the CPU path via bFallBackToCPU for mixed-layer files.
* Adds FOpenExrHeaderReader::GetChannelNames() to OpenExrWrapper to enumerate all channel names from an EXR header.

#Synced-CL 52455647

diff --git a/Engine/Plugins/Media/ImgMedia/Source/ImgMedia/Internal/Loader/ImgMediaLoader.h b/Engine/Plugins/Media/ImgMedia/Source/ImgMedia/Internal/Loader/ImgMediaLoader.h
--- a/Engine/Plugins/Media/ImgMedia/Source/ImgMedia/Internal/Loader/ImgMediaLoader.h
+++ b/Engine/Plugins/Media/ImgMedia/Source/ImgMedia/Internal/Loader/ImgMediaLoader.h
@@ -81,13 +81,20 @@ public:
 	/**
 	 * Create and initialize a new instance.
 	 *
-	 * @param InScheduler The scheduler for image loading.
+	 * @param InScheduler               The scheduler for image loading.
+	 * @param InGlobalCache             The global frame cache shared across loaders.
+	 * @param InMipMapInfo              Mip map configuration, or null if mips are not used.
+	 * @param InSourceColorSettings     Color space and encoding overrides for the source, or null for defaults.
+	 * @param bInFillGapsInSequence     If true, gaps in the image sequence are filled with blank frames.
+	 * @param InLayerName               EXR layer to select. Empty string reads the default (un-prefixed) channels.
+	 * @param InSmartCacheSettings      Settings controlling the smart cache pre-fetching behavior.
 	 */
 	FImgMediaLoader(const TSharedRef<FImgMediaScheduler, ESPMode::ThreadSafe>& InScheduler,
 		const TSharedRef<FImgMediaGlobalCache, ESPMode::ThreadSafe>& InGlobalCache,
 		const TSharedPtr<FImgMediaMipMapInfo, ESPMode::ThreadSafe>& InMipMapInfo,
 		const TSharedPtr<FNativeMediaSourceColorSettings, ESPMode::ThreadSafe>& InSourceColorSettings,
 		bool bInFillGapsInSequence,
+		const FString& InLayerName,
 		const FImgMediaLoaderSmartCacheSettings& InSmartCacheSettings);
 
 	/** Virtual destructor. */
@@ -642,6 +649,9 @@ private:
 	/** Container for media source color settings. */
 	TSharedPtr<FNativeMediaSourceColorSettings, ESPMode::ThreadSafe> SourceColorSettings;
 
+	/** Name of the EXR layer to read. */
+	FString LayerName;
+
 	/** Width and height of the image sequence (in pixels) .*/
 	FIntPoint SequenceDim;
 
diff --git a/Engine/Plugins/Media/ImgMedia/Source/ImgMedia/Private/Assets/ImgMediaSource.cpp b/Engine/Plugins/Media/ImgMedia/Source/ImgMedia/Private/Assets/ImgMediaSource.cpp
--- a/Engine/Plugins/Media/ImgMedia/Source/ImgMedia/Private/Assets/ImgMediaSource.cpp
+++ b/Engine/Plugins/Media/ImgMedia/Source/ImgMedia/Private/Assets/ImgMediaSource.cpp
@@ -53,6 +53,7 @@ namespace UE::ImgMediaSource::Private
 UImgMediaSource::UImgMediaSource()
 	: IsPathRelativeToProjectRoot_DEPRECATED(false)
 	, FrameRateOverride(0, 0)
+	, LayerName(FString())
 	, bFillGapsInSequence(true)
 	, MipMapInfo(MakeShared<FImgMediaMipMapInfo, ESPMode::ThreadSafe>())
 	, NativeSourceColorSettings(MakeShared<FNativeMediaSourceColorSettings, ESPMode::ThreadSafe>())
@@ -222,6 +223,11 @@ FString UImgMediaSource::GetMediaOption(const FName& Key, const FString& Default
 		return ProxyOverride;
 	}
 
+	if (Key == ImgMedia::LayerNameOption)
+	{
+		return LayerName;
+	}
+
 	if (Key == UMediaPlayer::MediaInfoNameStartTimecodeValue.Resolve())
 	{
 		return StartTimecode.ToString();
@@ -253,6 +259,7 @@ bool UImgMediaSource::HasMediaOption(const FName& Key) const
 		(Key == ImgMedia::FrameRateOverrideDenonimatorOption) ||
 		(Key == ImgMedia::FrameRateOverrideNumeratorOption) ||
 		(Key == ImgMedia::ProxyOverrideOption) ||
+		(Key == ImgMedia::LayerNameOption) ||
 		(Key == ImgMedia::MipMapInfoOption) ||
 		(Key == ImgMedia::SourceColorSettingsOption))
 	{
diff --git a/Engine/Plugins/Media/ImgMedia/Source/ImgMedia/Private/ImgMediaPrivate.h b/Engine/Plugins/Media/ImgMedia/Source/ImgMedia/Private/ImgMediaPrivate.h
--- a/Engine/Plugins/Media/ImgMedia/Source/ImgMedia/Private/ImgMediaPrivate.h
+++ b/Engine/Plugins/Media/ImgMedia/Source/ImgMedia/Private/ImgMediaPrivate.h
@@ -46,6 +46,9 @@ namespace ImgMedia
 	/** Name of the ProxyOverride media option. */
 	static const FName ProxyOverrideOption("ProxyOverride");
 
+	/** Name of the LayerName media option. */
+	static const FName LayerNameOption("LayerName");
+
 	/** Name of the MipMapInfo media option. */
 	static const FName MipMapInfoOption("ImgMipMapInfo");
 
diff --git a/Engine/Plugins/Media/ImgMedia/Source/ImgMedia/Private/Loader/ImgMediaLoader.cpp b/Engine/Plugins/Media/ImgMedia/Source/ImgMedia/Private/Loader/ImgMediaLoader.cpp
--- a/Engine/Plugins/Media/ImgMedia/Source/ImgMedia/Private/Loader/ImgMediaLoader.cpp
+++ b/Engine/Plugins/Media/ImgMedia/Source/ImgMedia/Private/Loader/ImgMediaLoader.cpp
@@ -9,6 +9,7 @@
 #include "Misc/FrameRate.h"
 #include "HAL/FileManager.h"
 #include "HAL/PlatformProcess.h"
+#include "Hash/Fnv.h"
 #include "Misc/Paths.h"
 #include "Misc/QueuedThreadPool.h"
 #include "Misc/ScopeLock.h"
@@ -168,6 +169,7 @@ FImgMediaLoader::FImgMediaLoader(const TSharedRef<FImgMediaScheduler, ESPMode::T
 	const TSharedPtr<FImgMediaMipMapInfo, ESPMode::ThreadSafe>& InMipMapInfo,
 	const TSharedPtr<FNativeMediaSourceColorSettings, ESPMode::ThreadSafe>& InSourceColorSettings,
 	bool bInFillGapsInSequence,
+	const FString& InLayerName,
 	const FImgMediaLoaderSmartCacheSettings& InSmartCacheSettings)
 	: Frames(1)
 	, Initialized(false)
@@ -182,6 +184,7 @@ FImgMediaLoader::FImgMediaLoader(const TSharedRef<FImgMediaScheduler, ESPMode::T
 	, GlobalCache(InGlobalCache)
 	, MipMapInfo(InMipMapInfo)
 	, SourceColorSettings(InSourceColorSettings)
+	, LayerName(InLayerName)
 	, SequenceDim(FIntPoint::ZeroValue)
 	, SequenceDuration(FTimespan::Zero())
 	, SequenceFrameRate(0, 0)
@@ -823,7 +826,7 @@ bool FImgMediaLoader::LoadSequence(const FString& SequencePath, const FFrameRate
 	{
 #if IMGMEDIA_EXR_SUPPORTED_PLATFORM
 		// Differentiate between Uncompressed exr and the rest.
-		Reader = FExrImgMediaReader::GetReader(AsShared(), ImagePaths[0][0]);
+		Reader = FExrImgMediaReader::GetReader(AsShared(), ImagePaths[0][0], LayerName);
 #else
 		UE_LOG(LogImgMedia, Error, TEXT("EXR image sequences are currently supported on macOS and Windows only"));
 		return false;
@@ -857,7 +860,17 @@ bool FImgMediaLoader::LoadSequence(const FString& SequencePath, const FFrameRate
 		// Smart cache does not use the global cache.
 		UseGlobalCache = false;
 	}
-	SequenceName = FName(*SequencePath);
+	/** Include LayerName in the cache key so that different layer selections
+	 *  for the same sequence path map to separate cache buckets. Empty LayerName
+	 *  uses the path alone, preserving the existing key for the default case.
+	 * 
+	 *  The layer name is stored as its hash rather than the raw string because
+	 *  FName comparison is case-insensitive: embedding a uint32 keeps the key
+	 *  case-sensitive (numbers are unaffected by FName's case folding).
+	 */
+	SequenceName = LayerName.IsEmpty()
+		? FName(*SequencePath)
+		: FName(*FString::Printf(TEXT("%s?layer=%u"), *SequencePath, UE::HashStringFNV1a32(MakeStringView(LayerName))));
 
 	// fetch sequence attributes from first image
 	{
diff --git a/Engine/Plugins/Media/ImgMedia/Source/ImgMedia/Private/Player/ImgMediaPlayer.cpp b/Engine/Plugins/Media/ImgMedia/Source/ImgMedia/Private/Player/ImgMediaPlayer.cpp
--- a/Engine/Plugins/Media/ImgMedia/Source/ImgMedia/Private/Player/ImgMediaPlayer.cpp
+++ b/Engine/Plugins/Media/ImgMedia/Source/ImgMedia/Private/Player/ImgMediaPlayer.cpp
@@ -252,6 +252,7 @@ bool FImgMediaPlayer::Open(const FString& Url, const IMediaOptions* Options, con
 	TSharedPtr<FImgMediaMipMapInfo, ESPMode::ThreadSafe> MipMapInfo;
 	TSharedPtr<FNativeMediaSourceColorSettings, ESPMode::ThreadSafe> SourceColorSettings;
 	bool bFillGapsInSequence = true;
+	FString LayerName;
 	bool bIsSmartCacheEnabled = GetBoolOption(PlayerOptions, MediaPlayerOptionValues::ImgMediaSmartCacheEnabled(), false);
 	bool bIsSequencerEnvironment = GetNameOption(PlayerOptions, MediaPlayerOptionValues::Environment()).IsEqual(MediaPlayerOptionValues::Environment_Sequencer());
 	double SmartCacheTimeToLookAhead = GetDoubleOption(PlayerOptions, MediaPlayerOptionValues::ImgMediaSmartCacheTimeToLookAhead(), 0.0);
@@ -267,6 +268,7 @@ bool FImgMediaPlayer::Open(const FString& Url, const IMediaOptions* Options, con
 		FrameRateOverride.Denominator = Options->GetMediaOption(ImgMedia::FrameRateOverrideDenonimatorOption, 0LL);
 		FrameRateOverride.Numerator = Options->GetMediaOption(ImgMedia::FrameRateOverrideNumeratorOption, 0LL);
 		bFillGapsInSequence = Options->GetMediaOption(ImgMedia::FillGapsInSequenceOption, true);
+		LayerName = Options->GetMediaOption(ImgMedia::LayerNameOption, FString());
 		TSharedPtr<IMediaOptions::FDataContainer, ESPMode::ThreadSafe> DefaultValue;
 		TSharedPtr<IMediaOptions::FDataContainer, ESPMode::ThreadSafe> DataContainer = Options->GetMediaOption(ImgMedia::MipMapInfoOption, DefaultValue);
 		if (DataContainer.IsValid())
@@ -319,7 +321,7 @@ bool FImgMediaPlayer::Open(const FString& Url, const IMediaOptions* Options, con
 	// initialize image loader on a separate thread
 	FImgMediaLoaderSmartCacheSettings SmartCacheSettings(bIsSmartCacheEnabled, (float)SmartCacheTimeToLookAhead);
 	Loader = MakeShared<FImgMediaLoader, ESPMode::ThreadSafe>(Scheduler.ToSharedRef(),
-		GlobalCache.ToSharedRef(), MipMapInfo, SourceColorSettings, bFillGapsInSequence, SmartCacheSettings);
+		GlobalCache.ToSharedRef(), MipMapInfo, SourceColorSettings, bFillGapsInSequence, LayerName, SmartCacheSettings);
 	Scheduler->RegisterLoader(Loader.ToSharedRef());
 
 	const FString SequencePath = Url.RightChop(6);
diff --git a/Engine/Plugins/Media/ImgMedia/Source/ImgMedia/Private/Readers/ExrImgMediaReader.cpp b/Engine/Plugins/Media/ImgMedia/Source/ImgMedia/Private/Readers/ExrImgMediaReader.cpp
--- a/Engine/Plugins/Media/ImgMedia/Source/ImgMedia/Private/Readers/ExrImgMediaReader.cpp
+++ b/Engine/Plugins/Media/ImgMedia/Source/ImgMedia/Private/Readers/ExrImgMediaReader.cpp
@@ -15,6 +15,7 @@
 #include "IImgMediaModule.h"
 #include "Loader/ImgMediaLoader.h"
 #include "Assets/ImgMediaMipMapInfo.h"
+#include "Logging/StructuredLog.h"
 
 DECLARE_MEMORY_STAT(TEXT("EXR Reader Pool Memory."), STAT_ExrMediaReaderPoolMem, STATGROUP_ImgMediaPlugin);
 
@@ -28,8 +29,9 @@ static TAutoConsoleVariable<bool> CVarEnableUncompressedExrGpuReader(
 /* FExrImgMediaReader structors
  *****************************************************************************/
 
-FExrImgMediaReader::FExrImgMediaReader(const TSharedRef<FImgMediaLoader, ESPMode::ThreadSafe>& InLoader)
+FExrImgMediaReader::FExrImgMediaReader(const TSharedRef<FImgMediaLoader, ESPMode::ThreadSafe>& InLoader, const FString& InLayerName)
 	: LoaderPtr(InLoader)
+	, LayerName(InLayerName)
 {
 	const UImgMediaSettings* Settings = GetDefault<UImgMediaSettings>();
 	
@@ -169,7 +171,28 @@ bool FExrImgMediaReader::ReadFrame(int32 FrameId, const TMap<int32, FImgMediaTil
 	{
 		return false;
 	}
-	
+
+	/**
+	 * Resolve layer on first read (only when a layer name is configured).
+	 * Uses the double-checked locking pattern: the acquire load provides a
+	 * lock-free fast path for all frames after the first, while the inner
+	 * relaxed load re-checks under the lock to avoid redundant resolution
+	 * if two threads race on the very first frame.
+	 */
+	if (!LayerName.IsEmpty() && !bLayerResolved.load(std::memory_order_acquire))
+	{
+		FScopeLock RegionScopeLock(&LayerResolveCriticalSection);
+		if (!bLayerResolved.load(std::memory_order_relaxed))
+		{
+			const FString& FirstImage = Loader->GetImagePath(FrameId, 0);
+			if (!ResolveLayer(FirstImage))
+			{
+				UE_LOGFMT(LogImgMedia, Error, "Failed to resolve EXR layer '{0}'", LayerName);
+				return false;
+			}
+		}
+	}
+
 	// Get tile info.
 	int32 NumTilesX = Loader->GetNumTilesX();
 	int32 NumTilesY = Loader->GetNumTilesY();
@@ -197,9 +220,10 @@ bool FExrImgMediaReader::ReadFrame(int32 FrameId, const TMap<int32, FImgMediaTil
 			return false;
 		}
 
-		// allocate frame buffer
-		// FRgbaInputFile loads any exr as rgba 16 bit per channel.
-		// UncompressedSize is used for cache, therefore it needs to be set to the actual size of the buffer.
+		/** Allocate the frame buffer. FExrInputFile always outputs interleaved RGBA 16-bit half-float,
+		 *  so the buffer size is Dim.X * Dim.Y * 4 * sizeof(uint16) per mip level.
+		 *  UncompressedSize must reflect the actual allocation size as it is used for cache accounting.
+		 */
 		FrameInfo.UncompressedSize = GetMipBufferTotalSize(Dim, Loader->GetNumMipLevels() > 1);
 		SIZE_T BufferSize = FrameInfo.UncompressedSize;
 		void* Buffer = FMemory::Malloc(BufferSize, PLATFORM_CACHE_LINE_SIZE);
@@ -260,13 +284,11 @@ bool FExrImgMediaReader::ReadFrame(int32 FrameId, const TMap<int32, FImgMediaTil
 				{
 					UE_LOG(LogImgMedia, Error, TEXT("Non-GPU reader doesn't currently support natively-tiled EXR file %s"), *Image);
 				}
-				else
+				else if (LayerName.IsEmpty())
 				{
-					// Get for our frame/mip level.
 					FRgbaInputFile InputFile(Image, 2);
 					if (InputFile.HasInputFile())
 					{
-						// read frame data
 						InputFile.SetFrameBuffer(MipDataPtr, Dim);
 						InputFile.ReadPixels(0, Dim.Y - 1);
 						OutFrame->MipTilesPresent.Emplace(CurrentMipLevel, CurrentTileSelection);
@@ -274,7 +296,25 @@ bool FExrImgMediaReader::ReadFrame(int32 FrameId, const TMap<int32, FImgMediaTil
 					}
 					else
 					{
-						UE_LOG(LogImgMedia, Error, TEXT("Could not load %s"), *Image);
+						UE_LOGFMT(LogImgMedia, Error, "Could not load {0}", Image);
+					}
+				}
+				else
+				{
+					FExrInputFile InputFile(Image, 2);
+					if (InputFile.HasInputFile())
+					{
+						InputFile.ReadPixels(
+							MipDataPtr, Dim,
+							ResolvedChannelR, ResolvedChannelG,
+							ResolvedChannelB, ResolvedChannelA,
+							0, Dim.Y - 1);
+						OutFrame->MipTilesPresent.Emplace(CurrentMipLevel, CurrentTileSelection);
+						OutFrame->NumTilesRead++;
+					}
+					else
+					{
+						UE_LOGFMT(LogImgMedia, Error, "Could not load {0}", Image);
 					}
 				}
 			}
@@ -302,7 +342,7 @@ void FExrImgMediaReader::UncancelFrame(int32 FrameNumber)
 }
 
 /** Gets reader type (GPU vs CPU) depending on size of EXR and its compression. */
-TSharedPtr<IImgMediaReader, ESPMode::ThreadSafe> FExrImgMediaReader::GetReader(const TSharedRef <FImgMediaLoader, ESPMode::ThreadSafe>& InLoader, FString FirstImageInSequencePath)
+TSharedPtr<IImgMediaReader, ESPMode::ThreadSafe> FExrImgMediaReader::GetReader(const TSharedRef <FImgMediaLoader, ESPMode::ThreadSafe>& InLoader, FString FirstImageInSequencePath, const FString& InLayerName)
 {
 #if defined(PLATFORM_WINDOWS) && PLATFORM_WINDOWS
 	if (!FPaths::FileExists(FirstImageInSequencePath))
@@ -313,7 +353,7 @@ TSharedPtr<IImgMediaReader, ESPMode::ThreadSafe> FExrImgMediaReader::GetReader(c
 	FImgMediaFrameInfo Info;
 	if (!GetInfo(FirstImageInSequencePath, Info))
 	{
-		return MakeShareable(new FExrImgMediaReader(InLoader));
+		return MakeShareable(new FExrImgMediaReader(InLoader, InLayerName));
 	}
 
 	const bool bIsDevelopmentCustomFormat = Info.FormatName.Equals(TEXT("EXR CUSTOM"));
@@ -327,18 +367,192 @@ TSharedPtr<IImgMediaReader, ESPMode::ThreadSafe> FExrImgMediaReader::GetReader(c
 	// Check GetCompressionName of OpenExrWrapper for other compression names.
 	// todo: Add and test Vulkan support
 	if (GDynamicRHI && GDynamicRHI->GetInterfaceType() == ERHIInterfaceType::D3D12
-		&& Info.CompressionName == "Uncompressed" 
+		&& Info.CompressionName == "Uncompressed"
 		&& CVarEnableUncompressedExrGpuReader.GetValueOnAnyThread()
 		&& bIsOptimizedForGpu
 		)
 	{
-		TSharedRef<FExrImgMediaReaderGpu, ESPMode::ThreadSafe> GpuReader = 
-			MakeShared<FExrImgMediaReaderGpu, ESPMode::ThreadSafe>(InLoader);
+		TSharedRef<FExrImgMediaReaderGpu, ESPMode::ThreadSafe> GpuReader =
+			MakeShared<FExrImgMediaReaderGpu, ESPMode::ThreadSafe>(InLoader, InLayerName);
 		return GpuReader;
 	}
 #endif
 
-	return MakeShareable(new FExrImgMediaReader(InLoader));
+	return MakeShareable(new FExrImgMediaReader(InLoader, InLayerName));
+}
+
+bool FExrImgMediaReader::ResolveLayer(const FString& ImagePath)
+{
+	if (bLayerResolved.load(std::memory_order_relaxed))
+	{
+		return true;
+	}
+
+	FOpenExrHeaderReader HeaderReader(ImagePath);
+	if (!HeaderReader.HasInputFile())
+	{
+		return false;
+	}
+
+	// Get all channel names from the file.
+	TArray<FString> ChannelNames;
+	HeaderReader.GetChannelNames(ChannelNames);
+
+	if (ChannelNames.Num() == 0)
+	{
+		return false;
+	}
+
+	// Parse channels into a layer map: layer prefix -> array of suffixes.
+	// Channel "R" -> layer "", suffix "R"
+	// Channel "beauty.R" -> layer "beauty", suffix "R"
+	// Channel "Zdepth.Z" -> layer "Zdepth", suffix "Z"
+	TMap<FString, TArray<FString>> LayerMap;
+	for (const FString& ChannelName : ChannelNames)
+	{
+		FString Prefix;
+		FString Suffix;
+		if (ChannelName.Split(TEXT("."), &Prefix, &Suffix, ESearchCase::CaseSensitive, ESearchDir::FromEnd))
+		{
+			LayerMap.FindOrAdd(Prefix).Add(Suffix);
+		}
+		else
+		{
+			// Un-prefixed channel (default layer).
+			LayerMap.FindOrAdd(FString()).Add(ChannelName);
+		}
+	}
+
+	// Match the configured LayerName against available layers.
+	FString MatchedLayer;
+	bool bFound = false;
+
+	if (LayerName.IsEmpty())
+	{
+		// Empty = default layer (un-prefixed channels).
+		if (LayerMap.Contains(FString()))
+		{
+			MatchedLayer = FString();
+			bFound = true;
+		}
+	}
+	else if (LayerName.Contains(TEXT("*")) || LayerName.Contains(TEXT("?")))
+	{
+		// Wildcard matching - sort keys for deterministic first-alphabetical match.
+		TArray<FString> SortedKeys;
+		LayerMap.GetKeys(SortedKeys);
+		SortedKeys.Sort();
+
+		TArray<FString> AllMatches;
+		for (const FString& Key : SortedKeys)
+		{
+			if (!Key.IsEmpty() && Key.MatchesWildcard(LayerName))
+			{
+				AllMatches.Add(Key);
+			}
+		}
+
+		if (AllMatches.Num() > 1)
+		{
+			UE_LOGFMT(LogImgMedia, Warning, "EXR layer wildcard '{0}' matched {1} layers in {2}; using first alphabetical match '{3}'",
+				LayerName, AllMatches.Num(), ImagePath, AllMatches[0]);
+		}
+
+		if (AllMatches.Num() > 0)
+		{
+			MatchedLayer = AllMatches[0];
+			bFound = true;
+		}
+	}
+	else
+	{
+		// Exact match.
+		if (LayerMap.Contains(LayerName))
+		{
+			MatchedLayer = LayerName;
+			bFound = true;
+		}
+	}
+
+	if (!bFound)
+	{
+		UE_LOGFMT(LogImgMedia, Error, "EXR layer '{0}' not found in {1}. Available layers (set LayerName to the quoted value, or leave empty for the default):", LayerName, ImagePath);
+		for (const TPair<FString, TArray<FString>>& Pair : LayerMap)
+		{
+			UE_LOGFMT(LogImgMedia, Warning, "  '{0}' with channels: {1}{2}",
+				Pair.Key,
+				FString::Join(Pair.Value, TEXT(", ")),
+				Pair.Key.IsEmpty() ? FString(TEXT(" (default)")) : FString());
+		}
+		// Fall back to default layer.
+		if (LayerMap.Contains(FString()))
+		{
+			MatchedLayer = FString();
+			bFound = true;
+		}
+		else
+		{
+			return false;
+		}
+	}
+
+	// Get the suffixes for the matched layer.
+	const TArray<FString>& Suffixes = LayerMap[MatchedLayer];
+
+	// Build full channel name from layer prefix and suffix.
+	auto MakeFullChannelName = [&MatchedLayer](const FString& Suffix) -> FString
+	{
+		if (MatchedLayer.IsEmpty())
+		{
+			return Suffix;
+		}
+		return MatchedLayer + TEXT(".") + Suffix;
+	};
+
+	// Map suffixes to RGBA output slots.
+	ResolvedChannelR = FString();
+	ResolvedChannelG = FString();
+	ResolvedChannelB = FString();
+	ResolvedChannelA = FString();
+
+	// Check for standard RGBA suffixes.
+	bool bHasR = Suffixes.Contains(TEXT("R"));
+	bool bHasG = Suffixes.Contains(TEXT("G"));
+	bool bHasB = Suffixes.Contains(TEXT("B"));
+	bool bHasA = Suffixes.Contains(TEXT("A"));
+
+	if (bHasR || bHasG || bHasB)
+	{
+		// Standard RGB(A) layer.
+		if (bHasR) ResolvedChannelR = MakeFullChannelName(TEXT("R"));
+		if (bHasG) ResolvedChannelG = MakeFullChannelName(TEXT("G"));
+		if (bHasB) ResolvedChannelB = MakeFullChannelName(TEXT("B"));
+		if (bHasA) ResolvedChannelA = MakeFullChannelName(TEXT("A"));
+	}
+	else if (Suffixes.Num() == 1)
+	{
+		// Single non-standard channel (e.g., Z) -> grayscale.
+		FString FullName = MakeFullChannelName(Suffixes[0]);
+		ResolvedChannelR = FullName;
+		ResolvedChannelG = FullName;
+		ResolvedChannelB = FullName;
+		// Alpha stays empty (fill with 1.0).
+	}
+	else
+	{
+		// Multiple non-standard suffixes -> map to RGBA in order.
+		if (Suffixes.Num() > 0) ResolvedChannelR = MakeFullChannelName(Suffixes[0]);
+		if (Suffixes.Num() > 1) ResolvedChannelG = MakeFullChannelName(Suffixes[1]);
+		if (Suffixes.Num() > 2) ResolvedChannelB = MakeFullChannelName(Suffixes[2]);
+		if (Suffixes.Num() > 3) ResolvedChannelA = MakeFullChannelName(Suffixes[3]);
+	}
+
+	UE_LOGFMT(LogImgMedia, Log, "EXR layer resolved: '{0}' -> R='{1}', G='{2}', B='{3}', A='{4}'",
+		LayerName.IsEmpty() ? FString(TEXT("(default)")) : LayerName,
+		ResolvedChannelR, ResolvedChannelG, ResolvedChannelB, ResolvedChannelA);
+
+	bLayerResolved.store(true, std::memory_order_release);
+	return true;
 }
 
 /* FExrImgMediaReader implementation
diff --git a/Engine/Plugins/Media/ImgMedia/Source/ImgMedia/Private/Readers/ExrImgMediaReader.h b/Engine/Plugins/Media/ImgMedia/Source/ImgMedia/Private/Readers/ExrImgMediaReader.h
--- a/Engine/Plugins/Media/ImgMedia/Source/ImgMedia/Private/Readers/ExrImgMediaReader.h
+++ b/Engine/Plugins/Media/ImgMedia/Source/ImgMedia/Private/Readers/ExrImgMediaReader.h
@@ -4,6 +4,8 @@
 
 #include "ImgMediaPrivate.h"
 
+#include <atomic>
+
 #if IMGMEDIA_EXR_SUPPORTED_PLATFORM
 
 #include "IImgMediaReader.h"
@@ -22,8 +24,8 @@ class FExrImgMediaReader
 {
 public:
 
-	/** Default constructor. */
-	FExrImgMediaReader(const TSharedRef<FImgMediaLoader, ESPMode::ThreadSafe>& InLoader);
+	/** Constructor. */
+	FExrImgMediaReader(const TSharedRef<FImgMediaLoader, ESPMode::ThreadSafe>& InLoader, const FString& InLayerName = FString());
 	virtual ~FExrImgMediaReader();
 public:
 
@@ -36,7 +38,7 @@ public:
 
 public:
 	/** Gets reader type (GPU vs CPU) depending on size of EXR and its compression. */
-	static TSharedPtr<IImgMediaReader, ESPMode::ThreadSafe> GetReader(const TSharedRef<FImgMediaLoader, ESPMode::ThreadSafe>& InLoader, FString FirstImageInSequencePath);
+	static TSharedPtr<IImgMediaReader, ESPMode::ThreadSafe> GetReader(const TSharedRef<FImgMediaLoader, ESPMode::ThreadSafe>& InLoader, FString FirstImageInSequencePath, const FString& InLayerName = FString());
 
 protected:
 	enum EReadResult
@@ -86,11 +88,32 @@ protected:
 	 */
 	SIZE_T GetMipBufferTotalSize(FIntPoint Dim, bool bInHasMips);
 
+	/**
+	 * Resolves the configured LayerName to concrete RGBA channel names.
+	 * Reads the channel list from the given file and populates ResolvedChannel[R/G/B/A].
+	 *
+	 * @param ImagePath Path to an EXR file to read channel names from.
+	 * @return true if resolution succeeded.
+	 */
+	bool ResolveLayer(const FString& ImagePath);
+
 protected:
 	TSet<int32> CanceledFrames;
 	FCriticalSection CanceledFramesCriticalSection;
 	TMap<int32, FRgbaInputFile*> PendingFrames;
 	TWeakPtr<FImgMediaLoader, ESPMode::ThreadSafe> LoaderPtr;
+
+	/** Configured layer name (may contain wildcards). */
+	FString LayerName;
+	/** Resolved channel names for each RGBA output slot. */
+	FString ResolvedChannelR;
+	FString ResolvedChannelG;
+	FString ResolvedChannelB;
+	FString ResolvedChannelA;
+	/** Guards layer resolution state (separate from CanceledFramesCriticalSection). */
+	FCriticalSection LayerResolveCriticalSection;
+	/** True after layer has been resolved to concrete channel names. */
+	std::atomic_bool bLayerResolved{false};
 };
 
 
diff --git a/Engine/Plugins/Media/ImgMedia/Source/ImgMedia/Private/Readers/ExrImgMediaReaderGpu.cpp b/Engine/Plugins/Media/ImgMedia/Source/ImgMedia/Private/Readers/ExrImgMediaReaderGpu.cpp
--- a/Engine/Plugins/Media/ImgMedia/Source/ImgMedia/Private/Readers/ExrImgMediaReaderGpu.cpp
+++ b/Engine/Plugins/Media/ImgMedia/Source/ImgMedia/Private/Readers/ExrImgMediaReaderGpu.cpp
@@ -18,6 +18,7 @@
 
 #include "Engine/TextureRenderTarget2D.h"
 #include "GlobalShader.h"
+#include "Logging/StructuredLog.h"
 #include "CommonRenderResources.h"
 #include "BaseMediaSourceColorSettings.h"
 #include "Loader/ImgMediaLoader.h"
@@ -81,15 +82,15 @@ namespace {
 /* FExrImgMediaReaderGpu structors
  *****************************************************************************/
 
-FExrImgMediaReaderGpu::FExrImgMediaReaderGpu(const TSharedRef<FImgMediaLoader, ESPMode::ThreadSafe>& InLoader)
-	: FExrImgMediaReader(InLoader)
+FExrImgMediaReaderGpu::FExrImgMediaReaderGpu(const TSharedRef<FImgMediaLoader, ESPMode::ThreadSafe>& InLoader, const FString& InLayerName)
+	: FExrImgMediaReader(InLoader, InLayerName)
 	, LastTickedFrameCounter((uint64)-1)
 	, bIsShuttingDown(false)
 	, bFallBackToCPU(false)
 {
-
 }
 
+
 FExrImgMediaReaderGpu::~FExrImgMediaReaderGpu()
 {
 	FScopeLock ScopeLock(&MemoryPoolCriticalSection);
@@ -221,7 +222,7 @@ FExrImgMediaReader::EReadResult FExrImgMediaReaderGpu::ReadMip
 bool FExrImgMediaReaderGpu::ReadFrame(int32 FrameId, const TMap<int32, FImgMediaTileSelection>& InMipTiles, TSharedPtr<FImgMediaFrame, ESPMode::ThreadSafe> OutFrame)
 {
 	// Fall back to cpu?
-	if (bFallBackToCPU)
+	if (bFallBackToCPU.load(std::memory_order_acquire))
 	{
 		return FExrImgMediaReader::ReadFrame(FrameId, InMipTiles, OutFrame);
 	}
@@ -231,7 +232,23 @@ bool FExrImgMediaReaderGpu::ReadFrame(int32 FrameId, const TMap<int32, FImgMedia
 	{
 		return false;
 	}
-	
+
+	/**
+	 * The GPU raw block reader cannot perform per-channel selection. Any file
+	 * that warrants a LayerName almost certainly has more than four channels
+	 * total, which IsOptimizedForGpu() would already have rejected. Delegate
+	 * all layer-name reads to the CPU reader which handles them correctly.
+	 */
+	if (!LayerName.IsEmpty())
+	{
+		UE_LOGFMT(LogImgMedia, Warning,
+			"GPU reader does not support EXR layer selection (layer '{0}'). Falling back to CPU reader.",
+			LayerName);
+		bFallBackToCPU.store(true, std::memory_order_release);
+		
+		return FExrImgMediaReader::ReadFrame(FrameId, InMipTiles, OutFrame);
+	}
+
 	const FString& LargestImagePath = Loader->GetImagePath(FrameId, 0);
 	FImgMediaFrameInfo FrameInfo;
 	if (!GetInfo(LargestImagePath, FrameInfo, OutFrame))
@@ -311,7 +328,7 @@ bool FExrImgMediaReaderGpu::ReadFrame(int32 FrameId, const TMap<int32, FImgMedia
 					}
 
 					// Fall back to CPU.
-					bFallBackToCPU = true;
+					bFallBackToCPU.store(true, std::memory_order_release);
 
 					return FExrImgMediaReader::ReadFrame(ConverterParams.FrameId, InMipTiles, OutFrame);
 				}
diff --git a/Engine/Plugins/Media/ImgMedia/Source/ImgMedia/Private/Readers/ExrImgMediaReaderGpu.h b/Engine/Plugins/Media/ImgMedia/Source/ImgMedia/Private/Readers/ExrImgMediaReaderGpu.h
--- a/Engine/Plugins/Media/ImgMedia/Source/ImgMedia/Private/Readers/ExrImgMediaReaderGpu.h
+++ b/Engine/Plugins/Media/ImgMedia/Source/ImgMedia/Private/Readers/ExrImgMediaReaderGpu.h
@@ -61,8 +61,8 @@ class FExrImgMediaReaderGpu
 {
 public:
 
-	/** Default constructor. */
-	FExrImgMediaReaderGpu(const TSharedRef<FImgMediaLoader, ESPMode::ThreadSafe>& InLoader);
+	/** Constructor. */
+	FExrImgMediaReaderGpu(const TSharedRef<FImgMediaLoader, ESPMode::ThreadSafe>& InLoader, const FString& InLayerName = FString());
 	virtual ~FExrImgMediaReaderGpu();
 
 public:
@@ -131,8 +131,8 @@ private:
 	/** A flag indicating this reader is being destroyed, therefore memory should not be returned. */
 	bool bIsShuttingDown;
 
-	/** If true, then just use the CPU to read the file. */
-	bool bFallBackToCPU;
+	/** If true, then just use the CPU to read the file. Written once with memory_order_release, read lock-free with memory_order_acquire. */
+	std::atomic_bool bFallBackToCPU;
 };
 
 
diff --git a/Engine/Plugins/Media/ImgMedia/Source/ImgMedia/Public/ImgMediaSource.h b/Engine/Plugins/Media/ImgMedia/Source/ImgMedia/Public/ImgMediaSource.h
--- a/Engine/Plugins/Media/ImgMedia/Source/ImgMedia/Public/ImgMediaSource.h
+++ b/Engine/Plugins/Media/ImgMedia/Source/ImgMedia/Public/ImgMediaSource.h
@@ -67,6 +67,10 @@ public:
 	UPROPERTY(BlueprintReadWrite, EditAnywhere, Category=Sequence, AdvancedDisplay)
 	FString ProxyOverride;
 
+	/** (EXR only) Name of the layer to read. Empty reads the default (un-prefixed) channels. Supports wildcards (* and ?). */
+	UPROPERTY(BlueprintReadWrite, EditAnywhere, Category=Sequence, AdvancedDisplay, meta=(DisplayName="Layer Name (EXR)"))
+	FString LayerName;
+
 	/** If true, then any gaps in the sequence will be filled with blank frames. */
 	UPROPERTY(BlueprintReadWrite, EditAnywhere, Category = Sequence)
 	bool bFillGapsInSequence;
diff --git a/Engine/Plugins/Media/ImgMedia/Source/OpenExrWrapper/Private/OpenExrWrapper.cpp b/Engine/Plugins/Media/ImgMedia/Source/OpenExrWrapper/Private/OpenExrWrapper.cpp
--- a/Engine/Plugins/Media/ImgMedia/Source/OpenExrWrapper/Private/OpenExrWrapper.cpp
+++ b/Engine/Plugins/Media/ImgMedia/Source/OpenExrWrapper/Private/OpenExrWrapper.cpp
@@ -6,6 +6,7 @@
 #include "Containers/UnrealString.h"
 #include "Math/IntRect.h"
 #include "OpenExrWrapperLog.h"
+#include "Logging/StructuredLog.h"
 
 PRAGMA_DEFAULT_VISIBILITY_START
 THIRD_PARTY_INCLUDES_START
@@ -21,6 +22,8 @@ THIRD_PARTY_INCLUDES_START
 	#include "OpenEXR/ImfTiledInputFile.h"
 	#include "OpenEXR/ImfTiledOutputFile.h"
 	#include "OpenEXR/ImfTiledRgbaFile.h"
+	#include "OpenEXR/ImfInputFile.h"
+	#include "OpenEXR/ImfFrameBuffer.h"
 THIRD_PARTY_INCLUDES_END
 PRAGMA_DEFAULT_VISIBILITY_END
 
@@ -274,6 +277,208 @@ bool FRgbaInputFile::GetIntAttribute(const FString& Name, int32& Value)
 	return bIsAttributeFound;
 }
 
+/* FExrInputFile
+ *****************************************************************************/
+
+FExrInputFile::FExrInputFile(const FString& FilePath, uint16 ThreadCount)
+{
+	try
+	{
+		InputFile = new Imf::InputFile(TCHAR_TO_ANSI(*FilePath), ThreadCount);
+	}
+	catch (std::exception const& Exception)
+	{
+		UE_LOGFMT(LogOpenEXRWrapper, Error, "Cannot load EXR file: {0}", Exception.what());
+		InputFile = nullptr;
+	}
+}
+
+
+FExrInputFile::~FExrInputFile()
+{
+	delete (Imf::InputFile*)InputFile;
+}
+
+
+bool FExrInputFile::HasInputFile() const
+{
+	return InputFile != nullptr;
+}
+
+
+FIntPoint FExrInputFile::GetDataWindow() const
+{
+	if (InputFile == nullptr)
+	{
+		return FIntPoint::ZeroValue;
+	}
+	Imath::Box2i Win = ((Imf::InputFile*)InputFile)->header().dataWindow();
+	return FIntPoint(
+		Win.max.x - Win.min.x + 1,
+		Win.max.y - Win.min.y + 1
+	);
+}
+
+
+void FExrInputFile::ReadPixels(
+	void* Buffer,
+	const FIntPoint& BufferDim,
+	const FString& ChannelR,
+	const FString& ChannelG,
+	const FString& ChannelB,
+	const FString& ChannelA,
+	int32 StartY,
+	int32 EndY)
+{
+	if (InputFile == nullptr)
+	{
+		return;
+	}
+
+	try
+	{
+		Imf::InputFile* File = (Imf::InputFile*)InputFile;
+		Imath::Box2i Win = File->header().dataWindow();
+
+		// Validate that our buffer is large enough for the EXR data window.
+		const FIntPoint DataDim(Win.max.x - Win.min.x + 1, Win.max.y - Win.min.y + 1);
+		if (!ensureMsgf(BufferDim.X >= DataDim.X && BufferDim.Y >= DataDim.Y,
+			TEXT("ReadPixels: buffer (%dx%d) is smaller than EXR data window (%dx%d)"),
+			BufferDim.X, BufferDim.Y, DataDim.X, DataDim.Y))
+		{
+			return;
+		}
+
+		/**
+		 * This function always writes interleaved RGBA half-float output:
+		 * [R,G,B,A, R,G,B,A, ...], 8 bytes per pixel.
+		 * The strides below define that layout regardless of how channels
+		 * are stored in the EXR file (planar, tiled, etc.).
+		 */
+		const int32 PixelStride = 4 * sizeof(uint16);
+		const int32 LineStride = BufferDim.X * PixelStride;
+
+		// Base pointer adjusted for data window origin so that
+		// pixel (Win.min.x, Win.min.y) maps to Buffer[0].
+		char* Base = (char*)Buffer - Win.min.x * PixelStride - Win.min.y * LineStride;
+
+		Imf::FrameBuffer FrameBuffer;
+
+		// Helper to add a channel slice to the frame buffer.
+		// SlotOffset: byte offset of this channel within a pixel (0=R, 2=G, 4=B, 6=A).
+		auto AddChannelSlice = [&](const FString& ChannelName, int32 SlotOffset, double SlotFillValue)
+		{
+			if (!ChannelName.IsEmpty())
+			{
+				// Only insert if not already present (same channel mapped to multiple slots
+				// is handled via post-read copy).
+				if (!FrameBuffer.findSlice(StringCast<ANSICHAR>(*ChannelName).Get()))
+				{
+					FrameBuffer.insert(
+						StringCast<ANSICHAR>(*ChannelName).Get(),
+						Imf::Slice(
+							Imf::HALF,
+							Base + SlotOffset,
+							PixelStride,
+							LineStride,
+							1, 1, // x/ySampling
+							SlotFillValue
+						)
+					);
+				}
+			}
+		};
+
+		/**
+		 * SlotFillValue is passed as the fillValue argument to Imf::Slice.
+		 * OpenEXR uses it to fill the output buffer for any channel that is
+		 * registered in the FrameBuffer but absent from the file -- so alpha
+		 * is always written as 1.0 for RGB-only EXRs even when "A" is
+		 * explicitly requested. The post-read loop below handles only the
+		 * separate case where the channel name itself is empty (no file
+		 * channel mapped to that output slot at all).
+		 */
+		AddChannelSlice(ChannelR, 0, 0.0);
+		AddChannelSlice(ChannelG, 2, 0.0);
+		AddChannelSlice(ChannelB, 4, 0.0);
+		AddChannelSlice(ChannelA, 6, 1.0);
+
+		File->setFrameBuffer(FrameBuffer);
+		File->readPixels(StartY + Win.min.y, EndY + Win.min.y);
+
+		// Post-read: handle fill slots and grayscale (duplicate channel) copies.
+		struct FChannelSlot
+		{
+			const FString* Name;
+			int32 Offset;
+		};
+		const FChannelSlot Slots[] = {
+			{ &ChannelR, 0 },
+			{ &ChannelG, 2 },
+			{ &ChannelB, 4 },
+			{ &ChannelA, 6 },
+		};
+
+		// Post-read: fill empty slots and copy duplicate slots (grayscale).
+		// Data is interleaved RGBA so we can't use a flat Memset; we work line-by-line
+		// and stride over each pixel (4 uint16s = 8 bytes per pixel).
+		for (int32 i = 0; i < 4; ++i)
+		{
+			if (Slots[i].Name->IsEmpty())
+			{
+				// Fill slot: 0 for RGB, 1.0 (0x3C00) for A.
+				const uint16 FillValue = (i == 3) ? 0x3C00 : 0;
+				// Build a scratch line with the fill value at every pixel's slot position,
+				// then write it for each scanline via a strided loop.
+				uint16* PixelPtr = (uint16*)((char*)Buffer + StartY * LineStride + Slots[i].Offset);
+				const int32 NumLines = EndY - StartY + 1;
+				const int32 NumPixels = BufferDim.X * NumLines;
+				for (int32 p = 0; p < NumPixels; ++p)
+				{
+					*PixelPtr = FillValue;
+					PixelPtr += 4; // stride to next pixel's same channel slot
+				}
+				continue;
+			}
+
+			// Check if this channel was already written to a prior slot (grayscale case).
+			int32 SourceSlot = -1;
+			for (int32 j = 0; j < i; ++j)
+			{
+				if (!Slots[j].Name->IsEmpty() && *Slots[j].Name == *Slots[i].Name)
+				{
+					SourceSlot = j;
+					break;
+				}
+			}
+
+			if (SourceSlot >= 0)
+			{
+				// Copy the already-decoded channel into this duplicate slot, line by line.
+				for (int32 y = StartY; y <= EndY; ++y)
+				{
+					const uint16* SrcPtr = (const uint16*)((const char*)Buffer + y * LineStride + Slots[SourceSlot].Offset);
+					uint16* DstPtr = (uint16*)((char*)Buffer + y * LineStride + Slots[i].Offset);
+					for (int32 x = 0; x < BufferDim.X; ++x)
+					{
+						*DstPtr = *SrcPtr;
+						SrcPtr += 4;
+						DstPtr += 4;
+					}
+				}
+			}
+		}
+	}
+	catch (std::exception const& Exception)
+	{
+		const ANSICHAR* FileName = InputFile ? ((Imf::InputFile*)InputFile)->fileName() : "";
+		UE_LOGFMT(LogOpenEXRWrapper, Error, "Cannot read EXR file: {0} ({1})",
+			FileName,
+			Exception.what());
+	}
+}
+
+
 FBaseOutputFile::FBaseOutputFile(
 	const FIntPoint& DisplayWindowMin,
 	const FIntPoint& DisplayWindowMax,
diff --git a/Engine/Plugins/Media/ImgMedia/Source/OpenExrWrapper/Private/OpenExrWrapperContextReader.cpp b/Engine/Plugins/Media/ImgMedia/Source/OpenExrWrapper/Private/OpenExrWrapperContextReader.cpp
--- a/Engine/Plugins/Media/ImgMedia/Source/OpenExrWrapper/Private/OpenExrWrapperContextReader.cpp
+++ b/Engine/Plugins/Media/ImgMedia/Source/OpenExrWrapper/Private/OpenExrWrapperContextReader.cpp
@@ -239,6 +239,29 @@ bool FOpenExrHeaderReader::HasInputFile() const
 	return FileContext != nullptr;
 }
 
+void FOpenExrHeaderReader::GetChannelNames(TArray<FString>& OutChannelNames) const
+{
+	OutChannelNames.Reset();
+
+	if (!FileContext.IsValid())
+	{
+		return;
+	}
+
+	const exr_attribute_t* Attribute = nullptr;
+	exr_result_t Result = exr_get_attribute_by_name(*((exr_context_t*)FileContext.Get()), 0, "channels", &Attribute);
+
+	if (Result != EXR_ERR_SUCCESS || Attribute == nullptr)
+	{
+		return;
+	}
+
+	for (int32 ChannelId = 0; ChannelId < Attribute->chlist->num_channels; ChannelId++)
+	{
+		OutChannelNames.Add(ANSI_TO_TCHAR(Attribute->chlist->entries[ChannelId].name.str));
+	}
+}
+
 bool FOpenExrHeaderReader::GetIntAttribute(const FString& Name, int32& Value)
 {
 	const exr_attribute_t* Attribute = nullptr;
diff --git a/Engine/Plugins/Media/ImgMedia/Source/OpenExrWrapper/Public/OpenExrWrapper.h b/Engine/Plugins/Media/ImgMedia/Source/OpenExrWrapper/Public/OpenExrWrapper.h
--- a/Engine/Plugins/Media/ImgMedia/Source/OpenExrWrapper/Public/OpenExrWrapper.h
+++ b/Engine/Plugins/Media/ImgMedia/Source/OpenExrWrapper/Public/OpenExrWrapper.h
@@ -60,6 +60,54 @@ private:
 	void* InputFile;
 };
 
+/**
+ * Reads EXR files with per-channel control using the general Imf::InputFile API.
+ * Unlike FRgbaInputFile, this allows selecting specific named channels for each RGBA output slot.
+ */
+class FExrInputFile
+{
+public:
+
+	UE_API FExrInputFile(const FString& FilePath, uint16 ThreadCount);
+	UE_API ~FExrInputFile();
+
+public:
+
+	UE_API bool HasInputFile() const;
+	UE_API FIntPoint GetDataWindow() const;
+
+	/**
+	 * Read specific channels into an RGBA half-float (uint16) interleaved buffer.
+	 *
+	 * Each channel parameter is a full EXR channel name (e.g., "R", "beauty.G", "Zdepth.Z").
+	 * Empty string means fill with a default value (0 for RGB slots, 1 for A slot).
+	 * The same channel name may be used for multiple slots (e.g., all RGB = "Zdepth.Z" for grayscale).
+	 *
+	 * @param Buffer		Destination buffer, must hold BufferDim.X * BufferDim.Y * 4 * sizeof(uint16) bytes.
+	 * @param BufferDim		Width and height of the buffer (used for stride calculation).
+	 * @param ChannelR		EXR channel name to read into the R slot.
+	 * @param ChannelG		EXR channel name to read into the G slot.
+	 * @param ChannelB		EXR channel name to read into the B slot.
+	 * @param ChannelA		EXR channel name to read into the A slot.
+	 * @param StartY		First scanline to read (0-based, relative to data window).
+	 * @param EndY			Last scanline to read (inclusive, 0-based, relative to data window).
+	 */
+	UE_API void ReadPixels(
+		void* Buffer,
+		const FIntPoint& BufferDim,
+		const FString& ChannelR,
+		const FString& ChannelG,
+		const FString& ChannelB,
+		const FString& ChannelA,
+		int32 StartY,
+		int32 EndY);
+
+private:
+
+	void* InputFile;
+};
+
+
 class FOpenExrHeaderReader
 {
 public:
@@ -98,6 +146,13 @@ public:
 	 */
 	UE_API bool GetIntAttribute(const FString& Name, int32& Value);
 
+	/**
+	 * Get the names of all channels in the file.
+	 *
+	 * @param OutChannelNames Will contain all channel names (e.g., "R", "G", "B", "Zdepth.Z").
+	 */
+	UE_API void GetChannelNames(TArray<FString>& OutChannelNames) const;
+
 private:
 	TSharedPtr<void> FileContext;
 };
