[TmvMedia][ImgMedia] Add .tmv container muxing/demuxing support

Common types (TmvMediaTrackInfo.h):
- FTmvMediaTrackInfo: shared track descriptor (codec 4CC, timescale, dimensions, etc.)
- FTmvMediaTrackSampleInfo: shared sample metadata (DTS, PTS, duration, keyframe flag)
- ETmvMediaTrackType, ETmvMediaContainerResult: shared enums
- FTmvMediaMuxerTrackConfig/FTmvMediaDemuxerTrackInfo extend FTmvMediaTrackInfo
- FTmvMediaMuxerSample/FTmvMediaDemuxerSample extend FTmvMediaTrackSampleInfo

Muxer/demuxer interfaces:
- ITmvMediaMuxer: abstract muxer interface (Configure, AddTrack, Start, AddSample, Finalize)
- ITmvMediaDemuxer: abstract demuxer interface (OpenFile, ReadSampleInfo, SeekToSample)
- ITmvMediaMuxerFactory (with GetFileExtension), ITmvMediaDemuxerFactory
- Factory registration in ITmvMediaModule
- FTmvMediaEncoderOptions::GetCodecFourCC() for codec 4CC identification

TmvMediaMp4Utils module (new):
- FTmvMediaMp4Muxer: wraps IMP4RawMuxer with push-to-pull sample bridge
- FTmvMediaMp4Demuxer: wraps MP4Boxes for container parsing and sample iteration
- Registered via module startup, supports "tmv" container format

Transcoder refactoring:
- UTmvMediaTranscodeMuxer: now abstract base class
- UTmvMediaFileSequenceTranscodeMuxer: existing file-per-frame behavior (unchanged)
- UTmvMediaContainerTranscodeMuxer: muxes via ITmvMediaMuxer with async push-to-pull bridge
- ETmvMediaTranscodeOutputFormat setting (FileSequence/Container) in job settings
- Encoder passes track config to container muxer on first RequestMipInfos

Container playback in ImgMedia:
- FTmvImgMediaReader: refactored to abstract base with virtual CreateAccessUnit
- FTmvFileSequenceImgMediaReader: file-per-frame (CreateAccessUnit opens individual file)
- FTmvContainerImgMediaReader: container (CreateAccessUnit uses demuxer + byte-range file reader)
- FContainerFileAccessUnit: zero-copy byte-range reader within container file
- FImgMediaLoader: dedicated container path (no ImagePaths, uses NumContainerFrames)
- FImgMediaFactoryModule: GetPlayabilityConfidenceScore supports file:// scheme for .tmv
- UTmvFileMediaSourceFactory: asset factory for .tmv files
- IImgMediaModulePrivate::GetOrCreateMipMapInfo: per-source MipMapInfo manager
- FImgMediaPlayer::Open: supports file:// URL scheme for container sources

#Synced-CL 52638463

diff --git a/Engine/Plugins/Media/ImgMedia/Source/ImgMedia/Internal/Loader/ImgMediaLoader.h b/Engine/Plugins/Media/ImgMedia/Source/ImgMedia/Internal/Loader/ImgMediaLoader.h
--- a/Engine/Plugins/Media/ImgMedia/Source/ImgMedia/Internal/Loader/ImgMediaLoader.h
+++ b/Engine/Plugins/Media/ImgMedia/Source/ImgMedia/Internal/Loader/ImgMediaLoader.h
@@ -294,15 +294,25 @@ public:
 	};
 
 	/**
-	 * Get the number of images in a single mip level.
+	 * Get the number of images/frames in a single mip level.
 	 *
 	 * @return Number of images.
 	 */
 	int32 GetNumImages() const
 	{
+		if (bIsContainerFile)
+		{
+			return NumContainerFrames;
+		}
 		return ImagePaths.Num() > 0 ? ImagePaths[0].Num() : 0;
 	}
 
+	/** Whether the source is a container file (e.g. .tmv) rather than a file sequence. */
+	bool IsContainerFile() const
+	{
+		return bIsContainerFile;
+	}
+
 	/**
 	 * Get the number of tiles in the X direction.
 	 */
@@ -527,22 +537,22 @@ protected:
 	 * Get what mip level and tiles we should be using for a given frame.
 	 *
 	 * @param FrameIndex Frame to get mip level for.
-	 * @param OutTileSelection Will be filled in with what tiles are actually needed.
+	 * @param OutMipsAndTiles Will be filled in with what tiles are actually needed.
 	 */
-	void GetDesiredMipTiles(int32 FrameIndex, TMap<int32, FImgMediaTileSelection>& OutMipsAndTiles);
+	void GetDesiredMipTiles(int32 FrameIndex, TMap<int32, FImgMediaTileSelection>& OutMipsAndTiles) const;
 
 	/**
 	 * Get what mip level supporting readers should upscale.
 	 *
 	 * @return minimum mip level to upscale.
 	 */
-	int32 GetDesiredMinimumMipLevelToUpscale();
+	int32 GetDesiredMinimumMipLevelToUpscale() const;
 
 	/***
 	 * Modulos the time so that it is between 0 and SequenceDuration.
 	 * Handles negative numbers appropriately.
 	 */
-	FTimespan ModuloTime(FTimespan Time);
+	FTimespan ModuloTime(FTimespan Time) const;
 
 	/**
 	 * Helper function to get the number at the end of a string.
@@ -568,7 +578,7 @@ protected:
 	 *		4			7
 	 *		5			7
 	 */
-	int32 GetSkipFrame(int32 FrameIndex);
+	int32 GetSkipFrame(int32 FrameIndex) const;
 
 	/**
 	 * Updates how we throttle bandwidth depending on current settings.
@@ -601,9 +611,16 @@ private:
 	/**
 	 * Paths to each image for each mip map level in the currently opened sequence.
 	 * This is an array of mip levels, and each mip level is an array of image paths.
+	 * Not used for container files.
 	 */
 	TArray<TArray<FString>> ImagePaths;
 
+	/** Whether the source is a container file rather than a file sequence. */
+	bool bIsContainerFile = false;
+
+	/** Number of frames in the container (only valid when bIsContainerFile is true). */
+	int32 NumContainerFrames = 0;
+
 	/** Media information string. */
 	FString Info;
 
diff --git a/Engine/Plugins/Media/ImgMedia/Source/ImgMedia/Private/IImgMediaModulePrivate.h b/Engine/Plugins/Media/ImgMedia/Source/ImgMedia/Private/IImgMediaModulePrivate.h
new file mode 100644
--- /dev/null
+++ b/Engine/Plugins/Media/ImgMedia/Source/ImgMedia/Private/IImgMediaModulePrivate.h
@@ -0,0 +1,42 @@
+// Copyright Epic Games, Inc. All Rights Reserved.
+
+#pragma once
+
+#include "IImgMediaModule.h"
+
+class FImgMediaMipMapInfo;
+class UMediaSource;
+
+/**
+ * Extended Private Interface for the ImgMedia module.
+ */
+class IImgMediaModulePrivate
+	: public IImgMediaModule
+{
+public:
+	/**
+	 * Returns a pointer to the module instance if loaded, null otherwise.
+	 */
+	static inline IImgMediaModulePrivate* Get()
+	{
+		return FModuleManager::GetModulePtr<IImgMediaModulePrivate>("ImgMedia");
+	}
+	
+	/**
+	 * Returns existing instance of the module, loads it if it was not already loaded.
+	 * Beware of calling this during the shutdown phase, though. The module might have been unloaded already.
+	 */
+	static IImgMediaModulePrivate& GetOrLoad()
+	{
+		return FModuleManager::LoadModuleChecked<IImgMediaModulePrivate>("ImgMedia");
+	}
+
+	/**
+	 * Get or create a shared MipMapInfo for the given media source that can be reused by different players.
+	 * If the source already has an entry, returns it. Otherwise, creates a new one.
+	 *
+	 * @param InMediaSource Optional media source to look up. Can be null, in which case a new MipMapInfo is created but not stored in the map.
+	 * @return The MipMapInfo for this source.
+	 */
+	virtual TSharedPtr<FImgMediaMipMapInfo, ESPMode::ThreadSafe> GetOrCreateSharedMipMapInfo(const UMediaSource* InMediaSource) = 0;
+};
\ No newline at end of file
diff --git a/Engine/Plugins/Media/ImgMedia/Source/ImgMedia/Private/ImgMediaModule.cpp b/Engine/Plugins/Media/ImgMedia/Source/ImgMedia/Private/ImgMediaModule.cpp
--- a/Engine/Plugins/Media/ImgMedia/Source/ImgMedia/Private/ImgMediaModule.cpp
+++ b/Engine/Plugins/Media/ImgMedia/Source/ImgMedia/Private/ImgMediaModule.cpp
@@ -1,17 +1,21 @@
 // Copyright Epic Games, Inc. All Rights Reserved.
 
-#include "ImgMediaPrivate.h"
-
+#include "FileMediaSource.h"
+#include "Assets/ImgMediaMipMapInfo.h"
+#include "IImgMediaModule.h"
+#include "IImgMediaModulePrivate.h"
 #include "IMediaClock.h"
 #include "IMediaModule.h"
-#include "Misc/QueuedThreadPool.h"
-
 #include "ImgMediaGlobalCache.h"
-#include "Player/ImgMediaPlayer.h"
+#include "ImgMediaPrivate.h"
 #include "ImgMediaSceneViewExtension.h"
-#include "Scheduler/ImgMediaScheduler.h"
 #include "ImgMediaSource.h"
-#include "IImgMediaModule.h"
+#include "ITmvMediaModule.h"
+#include "Decoder/ITmvMediaDecoderFactory.h"
+#include "Decoder/ITmvMediaDemuxerFactory.h"
+#include "Misc/QueuedThreadPool.h"
+#include "Player/ImgMediaPlayer.h"
+#include "Scheduler/ImgMediaScheduler.h"
 
 CSV_DEFINE_CATEGORY_MODULE(IMGMEDIA_API, ImgMedia, false);
 DEFINE_LOG_CATEGORY(LogImgMedia);
@@ -94,7 +98,7 @@ FQueuedThreadPool* GetImgMediaThreadPoolSlow()
  * Implements the AVFMedia module.
  */
 class FImgMediaModule
-	: public IImgMediaModule
+	: public IImgMediaModulePrivate
 {
 public:
 
@@ -126,6 +130,48 @@ public:
 	{
 		return SceneViewExtension;
 	}
+	
+	//~ Begin IImgMediaModulePrivate
+
+	virtual TSharedPtr<FImgMediaMipMapInfo, ESPMode::ThreadSafe> GetOrCreateSharedMipMapInfo(const UMediaSource* InMediaSource) override
+	{
+		FScopeLock Lock(&MipMapInfoMapCS);
+
+		// Prune stale entries while we have the lock.
+		for (auto It = MipMapInfoMap.CreateIterator(); It; ++It)
+		{
+			if (!It.Key().IsValid())
+			{
+				It.RemoveCurrent();
+			}
+		}
+
+		TSharedPtr<FImgMediaMipMapInfo, ESPMode::ThreadSafe> MipMapInfo;
+
+		TWeakObjectPtr<const UMediaSource> Key(InMediaSource);
+
+		if (InMediaSource != nullptr)
+		{
+			if (TSharedPtr<FImgMediaMipMapInfo, ESPMode::ThreadSafe>* FoundMipMapInfo = MipMapInfoMap.Find(Key))
+			{
+				MipMapInfo = *FoundMipMapInfo;
+			}
+		}
+
+		if (!MipMapInfo.IsValid())
+		{
+			MipMapInfo = MakeShared<FImgMediaMipMapInfo, ESPMode::ThreadSafe>();
+
+			if (InMediaSource != nullptr)
+			{
+				MipMapInfoMap.Add(Key, MipMapInfo);
+			}
+		}
+
+		return MipMapInfo;
+	}
+
+	//~ End IImgMediaModulePrivate
 
 public:
 
@@ -135,13 +181,15 @@ public:
 	{
 		// Register media source spawners.
 		FMediaSourceSpawnDelegate SpawnDelegate =
-			FMediaSourceSpawnDelegate::CreateStatic(&FImgMediaModule::SpawnMediaSourceForString);
+			FMediaSourceSpawnDelegate::CreateRaw(this, &FImgMediaModule::SpawnMediaSourceForString);
+
 		for (const FString& Ext : FileExtensions)
 		{
 			UMediaSource::RegisterSpawnFromFileExtension(Ext, SpawnDelegate);
 		}
 
 		FCoreDelegates::GetOnPostEngineInit().AddRaw(this, &FImgMediaModule::OnPostEngineInit);
+		FCoreDelegates::OnAllModuleLoadingPhasesComplete.AddRaw(this, &FImgMediaModule::OnAllModuleLoadingPhasesComplete);
 		FCoreDelegates::OnEnginePreExit.AddRaw(this, &FImgMediaModule::OnEnginePreExit);
 
 		ConsoleCmds.Add(IConsoleManager::Get().RegisterConsoleCommand(
@@ -166,10 +214,13 @@ public:
 			UMediaSource::UnregisterSpawnFromFileExtension(Ext);
 		}
 
+		UnregisterTmvMediaFactories();
+
 		Scheduler.Reset();
 		GlobalCache.Reset();
 
 		FCoreDelegates::OnEnginePreExit.RemoveAll(this);
+		FCoreDelegates::OnAllModuleLoadingPhasesComplete.RemoveAll(this);
 		FCoreDelegates::GetOnPostEngineInit().RemoveAll(this);
 
 #if USE_IMGMEDIA_DEALLOC_POOL
@@ -184,6 +235,74 @@ private:
 		SceneViewExtension = FSceneViewExtensions::NewExtension<FImgMediaSceneViewExtension>();
 	}
 
+	void OnAllModuleLoadingPhasesComplete()
+	{
+		RegisterTmvMediaFactories();
+	}
+
+	/** Register media source spawners from TmvMedia decoder and demuxer factories. */
+	void RegisterTmvMediaFactories()
+	{
+		const ITmvMediaModule* TmvMediaModule = ITmvMediaModule::Get();
+		if (!TmvMediaModule)
+		{
+			return;
+		}
+
+		FMediaSourceSpawnDelegate SpawnDelegate = FMediaSourceSpawnDelegate::CreateRaw(this, &FImgMediaModule::SpawnMediaSourceForString);
+
+		// Register all the tmv decoder factories for file sequence.
+		{
+			TArray<TWeakPtr<ITmvMediaDecoderFactory, ESPMode::ThreadSafe>> DecoderFactories;
+			TmvMediaModule->GetDecoderFactories(DecoderFactories);
+
+			for (const TWeakPtr<ITmvMediaDecoderFactory, ESPMode::ThreadSafe>& DecoderFactoryWeak : DecoderFactories)
+			{
+				if (const TSharedPtr<ITmvMediaDecoderFactory> DecoderFactory = DecoderFactoryWeak.Pin())
+				{
+					for (const FString& Extension : DecoderFactory->GetSupportedFileExtensions())
+					{
+						if (!TmvMediaFileExtensions.Contains(Extension))
+						{
+							UMediaSource::RegisterSpawnFromFileExtension(Extension, SpawnDelegate);
+							TmvMediaFileExtensions.Add(Extension);
+						}
+					}
+				}
+			}
+		}
+
+		// Register all tmv demuxer factories for container.
+		{
+			TArray<TWeakPtr<ITmvMediaDemuxerFactory, ESPMode::ThreadSafe>> DemuxerFactories;
+			TmvMediaModule->GetDemuxerFactories(DemuxerFactories);
+
+			for (const TWeakPtr<ITmvMediaDemuxerFactory, ESPMode::ThreadSafe>& DemuxerFactoryWeak : DemuxerFactories)
+			{
+				if (const TSharedPtr<ITmvMediaDemuxerFactory> DemuxerFactory = DemuxerFactoryWeak.Pin())
+				{
+					for (const FString& Extension : DemuxerFactory->GetSupportedContainerFormats())
+					{
+						if (!TmvMediaFileExtensions.Contains(Extension))
+						{
+							UMediaSource::RegisterSpawnFromFileExtension(Extension, SpawnDelegate);
+							TmvMediaFileExtensions.Add(Extension);
+							TmvMediaContainerExtensions.Add(Extension);
+						}
+					}
+				}
+			}
+		}
+	}
+	
+	void UnregisterTmvMediaFactories()
+	{
+		for (const FString& Extension : TmvMediaFileExtensions)
+		{
+			UMediaSource::UnregisterSpawnFromFileExtension(Extension);
+		}
+	}
+
 	void OnEnginePreExit()
 	{
 		SceneViewExtension.Reset();
@@ -229,13 +348,20 @@ private:
 	 * @param	MediaPath		File path to the media.
 	 * @param	Outer			Outer to use for this object.
 	 */
-	static UMediaSource* SpawnMediaSourceForString(const FString& MediaPath, UObject* Outer)
+	UMediaSource* SpawnMediaSourceForString(const FString& MediaPath, UObject* Outer) const
 	{
-		TObjectPtr<UImgMediaSource> MediaSource = 
-			NewObject<UImgMediaSource>(Outer, NAME_None, RF_Transactional);
-		MediaSource->SetSequencePath(MediaPath);
+		// For supported container extensions, we will create a UFileMediaSource.
+		if (TmvMediaContainerExtensions.Contains(FPaths::GetExtension(MediaPath)))
+		{
+			TObjectPtr<UFileMediaSource> FileMediaSource = NewObject<UFileMediaSource>(Outer, NAME_None, RF_Transactional);
+			FileMediaSource->SetFilePath(MediaPath);
+			return FileMediaSource;
+		}
 
-		return MediaSource;
+		// All other extensions are file sequence.
+		TObjectPtr<UImgMediaSource> ImgMediaSource = NewObject<UImgMediaSource>(Outer, NAME_None, RF_Transactional);
+		ImgMediaSource->SetSequencePath(MediaPath);
+		return ImgMediaSource;
 	}
 
 	TSharedPtr<FImgMediaScheduler, ESPMode::ThreadSafe> Scheduler;
@@ -251,11 +377,29 @@ private:
 		TEXT("jpg"),
 		TEXT("jpeg"),
 		TEXT("png"),
-		TEXT("dds"),
-		TEXT("apv1"),	// todo: Tmp Fixme - Tmv file extensions need to be derived from TmvDecoderFactories.
+		TEXT("dds")
 	};
+	
+	/** Extra file extensions supported by Tmv Media Module factories. */
+	TSet<FString> TmvMediaFileExtensions;
 
+	/** File extensions supported by Tmv Container factories. */
+	TSet<FString> TmvMediaContainerExtensions;
+
+	/** Registered console commands. */
 	TArray<IConsoleObject*> ConsoleCmds;
+
+	/**
+	 * Global map that associates UMediaSource objects with FImgMediaMipMapInfo instances.
+	 * This allows sources that don't natively carry a MipMapInfo (e.g. UFileMediaSource)
+	 * to still participate in the ImgMedia mip/tile pipeline.
+	 *
+	 * MipMapInfo instances are created on first access and kept alive as long as the
+	 * UMediaSource is alive. Stale entries (for garbage-collected sources) are pruned
+	 * on subsequent accesses.
+	 */
+	FCriticalSection MipMapInfoMapCS;
+	TMap<TWeakObjectPtr<const UMediaSource>, TSharedPtr<FImgMediaMipMapInfo, ESPMode::ThreadSafe>> MipMapInfoMap;
 };
 
 
diff --git a/Engine/Plugins/Media/ImgMedia/Source/ImgMedia/Private/Loader/ImgMediaLoader.cpp b/Engine/Plugins/Media/ImgMedia/Source/ImgMedia/Private/Loader/ImgMediaLoader.cpp
--- a/Engine/Plugins/Media/ImgMedia/Source/ImgMedia/Private/Loader/ImgMediaLoader.cpp
+++ b/Engine/Plugins/Media/ImgMedia/Source/ImgMedia/Private/Loader/ImgMediaLoader.cpp
@@ -1,33 +1,32 @@
 // Copyright Epic Games, Inc. All Rights Reserved.
 
 #include "Loader/ImgMediaLoader.h"
-#include "ImgMediaGlobalCache.h"
-#include "ImgMediaPrivate.h"
 
 #include "Algo/Reverse.h"
+#include "BaseMediaSourceColorSettings.h"
 #include "Containers/SortedMap.h"
-#include "Misc/FrameRate.h"
 #include "HAL/FileManager.h"
 #include "HAL/PlatformProcess.h"
 #include "Hash/Fnv.h"
+#include "IImageWrapperModule.h"
+#include "IImgMediaModule.h"
+#include "ImgMediaGlobalCache.h"
+#include "ImgMediaLoaderUtils.h"
+#include "ImgMediaLoaderWork.h"
+#include "ImgMediaPrivate.h"
+#include "Misc/FrameRate.h"
 #include "Misc/Paths.h"
 #include "Misc/QueuedThreadPool.h"
 #include "Misc/ScopeLock.h"
 #include "Modules/ModuleManager.h"
-#include "UObject/Class.h"
-#include "UObject/UObjectGlobals.h"
-
+#include "Player/ImgMediaTextureSample.h"
+#include "Readers/DdsImgMediaReader.h"
 #include "Readers/GenericImgMediaReader.h"
-#include "IImageWrapperModule.h"
-#include "IImgMediaModule.h"
-#include "ImgMediaLoaderUtils.h"
 #include "Readers/IImgMediaReader.h"
-#include "Readers/DdsImgMediaReader.h"
 #include "Readers/TmvImgMediaReader.h"
-#include "ImgMediaLoaderWork.h"
-#include "BaseMediaSourceColorSettings.h"
 #include "Scheduler/ImgMediaScheduler.h"
-#include "Player/ImgMediaTextureSample.h"
+#include "UObject/Class.h"
+#include "UObject/UObjectGlobals.h"
 
 #if IMGMEDIA_EXR_SUPPORTED_PLATFORM
 	#include "Readers/ExrImgMediaReader.h"
@@ -801,54 +800,94 @@ bool FImgMediaLoader::LoadSequence(const FString& SequencePath, const FFrameRate
 		return false;
 	}
 
-	// locate image sequence files
-	TArray<FString> FoundPaths;
-	FindFiles(SequencePath, FoundPaths);
-	if (FoundPaths.Num() == 0)
+	// Check if SequencePath is a container file (e.g. .tmv) rather than a directory.
+	bIsContainerFile = FTmvImgMediaReader::IsContainerFileExtensionSupported(FPaths::GetExtension(SequencePath))
+		&& FPaths::FileExists(SequencePath);
+
+	TSharedPtr<FTmvImgMediaReader, ESPMode::ThreadSafe> ContainerReader;
+
+	if (bIsContainerFile)
 	{
-		UE_LOGF(LogImgMedia, Error, "The directory %ls does not contain any image files", *SequencePath);
-		return false;
+		// Container file: create the reader first, get frame count from the demuxer track info.
+		ContainerReader = FTmvImgMediaReader::CreateReader(AsShared(), SequencePath);
+		if (!ContainerReader.IsValid())
+		{
+			UE_LOGF(LogImgMedia, Error, "Failed to create container reader for %ls", *SequencePath);
+			return false;
+		}
+
+		const FTmvMediaDemuxerTrackInfo* VideoTrackInfo = ContainerReader->GetVideoTrackInfo();
+
+		if (!VideoTrackInfo)
+		{
+			UE_LOGF(LogImgMedia, Error, "Container reader for %ls has no video track.", *SequencePath);
+			return false;
+		}
+		
+		NumContainerFrames = static_cast<int32>(VideoTrackInfo->NumSamples);
+
+		if (NumContainerFrames <= 0)
+		{
+			UE_LOGF(LogImgMedia, Error, "Container %ls has no frames", *SequencePath);
+			return false;
+		}
 	}
-	ImagePaths.Emplace(FoundPaths);
+	else
+	{
+		// File sequence: scan directory for individual frame files.
+		TArray<FString> FoundPaths;
+		FindFiles(SequencePath, FoundPaths);
+		if (FoundPaths.Num() == 0)
+		{
+			UE_LOGF(LogImgMedia, Error, "The directory %ls does not contain any image files", *SequencePath);
+			return false;
+		}
+		ImagePaths.Emplace(MoveTemp(FoundPaths));
 
-	// Get mips.
-	FindMips(SequencePath);
+		// Get mips.
+		FindMips(SequencePath);
+	}
 
 	// Default to a 1x1 tile
 	TilingDescription.TileNum = FIntPoint(1, 1);
 
 	FScopeLock Lock(&CriticalSection);
+	
+	Reader = ContainerReader;
 
-	// create image reader
-	const FString FirstExtension = FPaths::GetExtension(ImagePaths[0][0]);
-
-	if (FirstExtension == TEXT("exr"))
+	// create image reader (if not already created for container path)
+	if (!Reader.IsValid())
 	{
+		const FString FirstExtension = FPaths::GetExtension(ImagePaths[0][0]);
+
+		if (FirstExtension == TEXT("exr"))
+		{
 #if IMGMEDIA_EXR_SUPPORTED_PLATFORM
-		// Differentiate between Uncompressed exr and the rest.
-		Reader = FExrImgMediaReader::GetReader(AsShared(), ImagePaths[0][0], LayerName);
+			// Differentiate between Uncompressed exr and the rest.
+			Reader = FExrImgMediaReader::GetReader(AsShared(), ImagePaths[0][0], LayerName);
 #else
-		UE_LOGF(LogImgMedia, Error, "EXR image sequences are currently supported on macOS and Windows only");
-		return false;
+			UE_LOGF(LogImgMedia, Error, "EXR image sequences are currently supported on macOS, Linux and Windows only");
+			return false;
 #endif
-	}
-	else if (FirstExtension == TEXT("dds"))
-	{
-		Reader = MakeShared<FDdsImgMediaReader>(AsShared());
-	}
-	else if (FTmvImgMediaReader::IsFileExtensionSupported(FirstExtension))
-	{
-		Reader = FTmvImgMediaReader::CreateReader(AsShared(), ImagePaths[0][0]);
-	}
-	else
-	{
-		Reader = MakeShareable(new FGenericImgMediaReader(AsShared()));
-	}
-	
-	if (Reader.IsValid() == false)
-	{
-		UE_LOGF(LogImgMedia, Error, "Reader is not valid for file %ls.", *ImagePaths[0][0]);
-		return false;
+		}
+		else if (FirstExtension == TEXT("dds"))
+		{
+			Reader = MakeShared<FDdsImgMediaReader>(AsShared());
+		}
+		else if (FTmvImgMediaReader::IsImageFileExtensionSupported(FirstExtension))
+		{
+			Reader = FTmvImgMediaReader::CreateReader(AsShared(), ImagePaths[0][0]);
+		}
+		else
+		{
+			Reader = MakeShareable(new FGenericImgMediaReader(AsShared()));
+		}
+
+		if (Reader.IsValid() == false)
+		{
+			UE_LOGF(LogImgMedia, Error, "Reader is not valid for %ls.", *ImagePaths[0][0]);
+			return false;
+		}
 	}
 
 	const UImgMediaSettings* Settings = GetDefault<UImgMediaSettings>();
@@ -881,7 +920,7 @@ bool FImgMediaLoader::LoadSequence(const FString& SequencePath, const FFrameRate
 		{
 			FirstFrameInfo = CachedFrame->Get()->GetInfo();
 		}
-		else if (!Reader->GetFrameInfo(ImagePaths[0][0], FirstFrameInfo))
+		else if (!Reader->GetFrameInfo(bIsContainerFile ? SequencePath : ImagePaths[0][0], FirstFrameInfo))
 		{
 			UE_LOGF(LogImgMedia, Error, "Failed to get frame information from first image in %ls", *SequencePath);
 			return false;
@@ -900,7 +939,7 @@ bool FImgMediaLoader::LoadSequence(const FString& SequencePath, const FFrameRate
 		return false;
 	}
 
-	NumMipLevels = FMath::Max(FirstFrameInfo.NumMipLevels, ImagePaths.Num());
+	NumMipLevels = FMath::Max(FirstFrameInfo.NumMipLevels, bIsContainerFile ? 1 : ImagePaths.Num());
 	NumBytesPerPixel = FirstFrameInfo.NumBytesPerPixel;
 
 	MinimumMipLevelToUpscale = -1;
@@ -922,7 +961,7 @@ bool FImgMediaLoader::LoadSequence(const FString& SequencePath, const FFrameRate
 
 	// With mips in separate folders, we scale the uncompressed size by its mip chain approximation.
 	// With mips in each frame, the readers already return a value that incorporates their size.
-	if (ImagePaths.Num() > 1)
+	if (MipsInSeparateFiles())
 	{
 		UncompressedSize = (UncompressedSize * 4) / 3;
 	}
@@ -990,7 +1029,7 @@ bool FImgMediaLoader::LoadSequence(const FString& SequencePath, const FFrameRate
 
 		// Apply look ahead/behind ratio.
 		const float LoadBehindScale = FMath::Clamp(Settings->CacheBehindPercentage, 0.0f, 100.0f) / 100.0f;
-		NumLoadBehind = (int32)(LoadBehindScale * NumFramesToLoad);
+		NumLoadBehind = FMath::TruncToInt32(LoadBehindScale * NumFramesToLoad);
 		NumLoadAhead = FMath::CeilToInt32((1.0f - LoadBehindScale) * NumFramesToLoad);
 	}
 
@@ -1038,7 +1077,7 @@ void FImgMediaLoader::FindFiles(const FString& SequencePath, TArray<FString>& Ou
 
 	// Same list of extensions as FGenericImgMediaReader::GetFrameInfoFromImage() with the addition of exr & dds
 	TSet<FString> SupportedImageFileExtensions = { TEXT("exr"), TEXT("jpg"), TEXT("jpeg"), TEXT("png"), TEXT("bmp"), TEXT("dds") };
-	SupportedImageFileExtensions.Append(FTmvImgMediaReader::GetSupportedFileExtensions());
+	SupportedImageFileExtensions.Append(FTmvImgMediaReader::GetSupportedImageFileExtensions());
 
 	TArray<FString> UnnumberedFiles;
 	TSortedMap<int32, FString> SortedFoundFiles;
@@ -1122,8 +1161,8 @@ void FImgMediaLoader::FindMips(const FString& SequencePath)
 
 	// Is this a mipmap level, i.e. something like 256x256?
 	int32 Index = 0;
-	bool FoundDelimeter = SequenceDir.FindLastChar(TEXT('x'), Index);
-	if (FoundDelimeter)
+	const bool bFoundDelimeter = SequenceDir.FindLastChar(TEXT('x'), Index);
+	if (bFoundDelimeter)
 	{
 		// Loop over all mip levels.
 		FString BaseMipDir = FPaths::GetPath(SequenceDir);
@@ -1386,7 +1425,7 @@ void FImgMediaLoader::Update(int32 PlayHeadFrame, float PlayRate, bool Loop)
 
 			if (NeedFrame)
 			{
-				UE_LOGF(LogImgMedia, VeryVerbose, "Loader %p: Invalidating frame %i at playhead %i due to missing mips or tiles.", this, FrameNumber, PlayHeadFrame);
+				UE_LOGF(LogImgMedia, VeryVerbose, "Loader %p: Invalidating frame %i at play head %i due to missing mips or tiles.", this, FrameNumber, PlayHeadFrame);
 			}
 		}
 
@@ -1405,7 +1444,9 @@ void FImgMediaLoader::Update(int32 PlayHeadFrame, float PlayRate, bool Loop)
 			if (!FoundFrame)
 			{
 				// Do we actually have a frame for this?
-				if (ImagePaths[0][FrameNumber].Len() == 0)
+				// For file sequences, check if the path is empty (gap-filled placeholder).
+				// Containers always have valid frames for all sample indices.
+				if (!bIsContainerFile && ImagePaths[0][FrameNumber].Len() == 0)
 				{
 					TryAddEmptyFrame(FrameNumber);
 				}
@@ -1511,12 +1552,10 @@ void FImgMediaLoader::AddFrameToCache(int32 FrameNumber, const TSharedPtr<FImgMe
 		UE_LOGF(LogImgMedia, VeryVerbose, "Loader %p: Loaded frame %i", this, FrameNumber);
 		if (UseGlobalCache)
 		{
-			const TSharedPtr<FImgMediaFrame, ESPMode::ThreadSafe>* ExistingFrame;
-			ExistingFrame = GlobalCache->FindAndTouch(SequenceName, FrameNumber);
-			if (ExistingFrame == nullptr)
+			if (GlobalCache->FindAndTouch(SequenceName, FrameNumber) == nullptr)
 			{
-				const bool bHasMipMapsInSeparateFolders = ImagePaths.Num() > 1;
-				GlobalCache->AddFrame(ImagePaths[0][FrameNumber], SequenceName, FrameNumber, Frame, bHasMipMapsInSeparateFolders);
+				const FString& FrameKey = bIsContainerFile ? SequenceName.ToString() : ImagePaths[0][FrameNumber];
+				GlobalCache->AddFrame(FrameKey, SequenceName, FrameNumber, Frame, MipsInSeparateFiles());
 			}
 		}
 		else
@@ -1526,7 +1565,7 @@ void FImgMediaLoader::AddFrameToCache(int32 FrameNumber, const TSharedPtr<FImgMe
 	}
 }
 
-void FImgMediaLoader::GetDesiredMipTiles(int32 FrameIndex, TMap<int32, FImgMediaTileSelection>& OutMipsAndTiles)
+void FImgMediaLoader::GetDesiredMipTiles(int32 FrameIndex, TMap<int32, FImgMediaTileSelection>& OutMipsAndTiles) const
 {
 	// note: While currently unused, FrameIndex could become useful for recorded frames.
 
@@ -1544,7 +1583,7 @@ void FImgMediaLoader::GetDesiredMipTiles(int32 FrameIndex, TMap<int32, FImgMedia
 	}
 }
 
-int32 FImgMediaLoader::GetDesiredMinimumMipLevelToUpscale()
+int32 FImgMediaLoader::GetDesiredMinimumMipLevelToUpscale() const
 {
 	if (HasMipsOrTilesEnabled())
 	{
@@ -1555,7 +1594,7 @@ int32 FImgMediaLoader::GetDesiredMinimumMipLevelToUpscale()
 }
 
 
-FTimespan FImgMediaLoader::ModuloTime(FTimespan Time)
+FTimespan FImgMediaLoader::ModuloTime(FTimespan Time) const
 {
 	bool IsNegative = Time < FTimespan::Zero();
 
@@ -1577,7 +1616,7 @@ bool FImgMediaLoader::GetNumberAtEndOfString(int32 &Number, const FString& Strin
 	int32 Index = String.Len() - 1 - Extension.Len();
 	for (; Index >= 0; --Index)
 	{
-		bool bIsDigit = FChar::IsDigit(String[Index]);
+		const bool bIsDigit = FChar::IsDigit(String[Index]);
 		if (bIsDigit)
 		{
 			break;
@@ -1613,7 +1652,7 @@ bool FImgMediaLoader::GetNumberAtEndOfString(int32 &Number, const FString& Strin
 }
 
 
-int32 FImgMediaLoader::GetSkipFrame(int32 FrameIndex)
+int32 FImgMediaLoader::GetSkipFrame(int32 FrameIndex) const
 {
 	int32 Mult = (1 << SkipFramesLevel);
 	int32 Mask = ~(Mult - 1);
diff --git a/Engine/Plugins/Media/ImgMedia/Source/ImgMedia/Private/Player/ImgMediaPlayer.cpp b/Engine/Plugins/Media/ImgMedia/Source/ImgMedia/Private/Player/ImgMediaPlayer.cpp
--- a/Engine/Plugins/Media/ImgMedia/Source/ImgMedia/Private/Player/ImgMediaPlayer.cpp
+++ b/Engine/Plugins/Media/ImgMedia/Source/ImgMedia/Private/Player/ImgMediaPlayer.cpp
@@ -1,6 +1,8 @@
 // Copyright Epic Games, Inc. All Rights Reserved.
 
 #include "Player/ImgMediaPlayer.h"
+
+#include "IImgMediaModulePrivate.h"
 #include "ImgMediaPrivate.h"
 #include "ImgMediaSource.h"
 
@@ -16,12 +18,14 @@
 #include "Assets/ImgMediaMipMapInfo.h"
 #include "Scheduler/ImgMediaScheduler.h"
 #include "ImgMediaSettings.h"
+#include "ITmvMediaModule.h"
 #include "Player/ImgMediaTextureSample.h"
 
 #include "MediaPlayer.h"
 #include "MediaPlayerFacade.h"
 #include "MediaTexture.h"
 #include "MediaTextureTracker.h"
+#include "Decoder/ITmvMediaDemuxerFactory.h"
 
 #define LOCTEXT_NAMESPACE "FImgMediaPlayer"
 
@@ -205,7 +209,28 @@ bool FImgMediaPlayer::Open(const FString& Url, const IMediaOptions* Options, con
 {
 	Close();
 
-	if (Url.IsEmpty() || !Url.StartsWith(TEXT("img://")))
+	FString UrlScheme;
+	FString UrlPath;
+
+	if (Url.IsEmpty() || !Url.Split(TEXT("://"), &UrlScheme, &UrlPath))
+	{
+		return false;
+	}
+	
+	bool bIsSupportedContainer = false;
+	if (UrlScheme == TEXT("file"))
+	{
+		const FString Extension = FPaths::GetExtension(UrlPath);
+		if (const ITmvMediaModule* TmvModule = ITmvMediaModule::Get())
+		{
+			if (TmvModule->FindDemuxerFactoryForExtension(Extension).IsValid())
+			{
+				bIsSupportedContainer = true;
+			}
+		}
+	}
+	
+	if (!Url.StartsWith(TEXT("img://")) && !bIsSupportedContainer)
 	{
 		return false;
 	}
@@ -274,35 +299,41 @@ bool FImgMediaPlayer::Open(const FString& Url, const IMediaOptions* Options, con
 		if (DataContainer.IsValid())
 		{
 			MipMapInfo = StaticCastSharedPtr<FImgMediaMipMapInfo, IMediaOptions::FDataContainer, ESPMode::ThreadSafe>(DataContainer);
-			if (MipMapInfo.IsValid())
-			{
-				MipMapInfo->PruneStaleObjects();
+		}
 
-				if (ViewTexture)
-				{
-					MediaTextureWeakPtr = ViewTexture;
-					MipMapInfo->AddObjectsUsingThisMediaTexture(MediaTextureWeakPtr.Get());
-				}
-				else
+		if (!MipMapInfo.IsValid())
+		{
+			MipMapInfo = IImgMediaModulePrivate::GetOrLoad().GetOrCreateSharedMipMapInfo(Cast<const UMediaSource>(Options->ToUObject()));
+		}
+
+		if (MipMapInfo.IsValid())
+		{
+			MipMapInfo->PruneStaleObjects();
+
+			if (ViewTexture)
+			{
+				MediaTextureWeakPtr = ViewTexture;
+				MipMapInfo->AddObjectsUsingThisMediaTexture(MediaTextureWeakPtr.Get());
+			}
+			else
+			{
+				// Look through all the media textures we know about to find the one tied to this player
+				for (TWeakObjectPtr<UMediaTexture> TexturePtr : FMediaTextureTracker::Get().GetTextures())
 				{
-					// Look through all the media textures we know about to find the one tied to this player
-					for (TWeakObjectPtr<UMediaTexture> TexturePtr : FMediaTextureTracker::Get().GetTextures())
+					UMediaTexture* Texture = TexturePtr.Get();
+					if (Texture != nullptr)
 					{
-						UMediaTexture* Texture = TexturePtr.Get();
-						if (Texture != nullptr)
+						// Does this match the player?
+						UMediaPlayer* MediaPlayer = Texture->GetMediaPlayer();
+						if (MediaPlayer != nullptr)
 						{
-							// Does this match the player?
-							UMediaPlayer* MediaPlayer = Texture->GetMediaPlayer();
-							if (MediaPlayer != nullptr)
+							TSharedPtr<IMediaPlayer, ESPMode::ThreadSafe> Player = MediaPlayer->GetPlayerFacade()->GetPlayer();
+							if (Player.IsValid())
 							{
-								TSharedPtr<IMediaPlayer, ESPMode::ThreadSafe> Player = MediaPlayer->GetPlayerFacade()->GetPlayer();
-								if (Player.IsValid())
+								if (Player.Get() == this)
 								{
-									if (Player.Get() == this)
-									{
-										MediaTextureWeakPtr = TexturePtr;
-										MipMapInfo->AddObjectsUsingThisMediaTexture(MediaTextureWeakPtr.Get());
-									}
+									MediaTextureWeakPtr = TexturePtr;
+									MipMapInfo->AddObjectsUsingThisMediaTexture(MediaTextureWeakPtr.Get());
 								}
 							}
 						}
@@ -324,7 +355,7 @@ bool FImgMediaPlayer::Open(const FString& Url, const IMediaOptions* Options, con
 		GlobalCache.ToSharedRef(), MipMapInfo, SourceColorSettings, bFillGapsInSequence, LayerName, SmartCacheSettings);
 	Scheduler->RegisterLoader(Loader.ToSharedRef());
 
-	const FString SequencePath = Url.RightChop(6);
+	const FString SequencePath = UrlPath;
 
 	Async(EAsyncExecution::ThreadPool, [FrameRateOverride, LoaderPtr = TWeakPtr<FImgMediaLoader, ESPMode::ThreadSafe>(Loader), Proxy, SequencePath, Loop = ShouldLoop, bIsSequencerEnvironment]()
 	{
diff --git a/Engine/Plugins/Media/ImgMedia/Source/ImgMedia/Private/Readers/TmvImgMediaReader.cpp b/Engine/Plugins/Media/ImgMedia/Source/ImgMedia/Private/Readers/TmvImgMediaReader.cpp
--- a/Engine/Plugins/Media/ImgMedia/Source/ImgMedia/Private/Readers/TmvImgMediaReader.cpp
+++ b/Engine/Plugins/Media/ImgMedia/Source/ImgMedia/Private/Readers/TmvImgMediaReader.cpp
@@ -5,6 +5,8 @@
 #include "Async/Mutex.h"
 #include "Decoder/ITmvMediaDecoder.h"
 #include "Decoder/ITmvMediaDecoderFactory.h"
+#include "Decoder/ITmvMediaDemuxer.h"
+#include "Decoder/ITmvMediaDemuxerFactory.h"
 #include "HAL/FileManager.h"
 #include "ITmvMediaModule.h"
 #include "ImgMediaPrivate.h"
@@ -12,6 +14,7 @@
 #include "SampleConverter/TmvMediaTextureSampleConverter.h"
 #include "SampleConverter/TmvMediaFrameMipBufferPool.h"
 #include "Templates/SharedPointer.h"
+#include "Utils/TmvMediaUtils.h"
 
 namespace UE::ImgMedia
 {
@@ -101,6 +104,106 @@ namespace UE::ImgMedia
 		TUniquePtr<FArchive> FileReader;
 	};
 
+	/**
+	 * Implementation of an access unit interface for a byte range within a container file.
+	 * Reads directly from the container at [BaseOffset, BaseOffset + SampleSize).
+	 * Tell/Seek/Read operate relative to the sample start (0-based).
+	 */
+	class FContainerFileAccessUnit : public ITmvMediaDecoderAccessUnit
+	{
+	public:
+		FContainerFileAccessUnit(const FString& InContainerPath, int64 InBaseOffset, int64 InSampleSize, int32 InFrameId)
+		: FrameId(InFrameId)
+		, Filename(InContainerPath)
+		, BaseOffset(InBaseOffset)
+		, SampleSize(InSampleSize)
+		, FileReader(IFileManager::Get().CreateFileReader(*InContainerPath))
+		{
+			if (FileReader)
+			{
+				FileReader->Seek(BaseOffset);
+			}
+		}
+
+		virtual ~FContainerFileAccessUnit() override
+		{
+			if (FileReader)
+			{
+				FileReader->Close();
+			}
+		}
+
+		bool IsValid() const
+		{
+			return FileReader.IsValid() && SampleSize > 0;
+		}
+
+		//~ Begin ITmvMediaDecoderAccessUnit
+		virtual int64 Tell() const override
+		{
+			return FileReader.IsValid() ? FileReader->Tell() - BaseOffset : 0;
+		}
+
+		virtual int64 GetTotalSize() const override
+		{
+			return SampleSize;
+		}
+
+		virtual bool Seek(int64 InOffset) override
+		{
+			if (FileReader.IsValid())
+			{
+				FileReader->Seek(BaseOffset + InOffset);
+				return !FileReader->IsError();
+			}
+			return false;
+		}
+
+		virtual int64 Read(void* OutBuffer, int64 InSize) override
+		{
+			if (FileReader.IsValid())
+			{
+				// Clamp read to sample boundary.
+				const int64 CurrentRelative = FileReader->Tell() - BaseOffset;
+				const int64 Remaining = SampleSize - CurrentRelative;
+				const int64 ReadSize = FMath::Min(InSize, Remaining);
+				if (ReadSize <= 0)
+				{
+					return 0;
+				}
+				FileReader->Serialize(OutBuffer, ReadSize);
+				if (!FileReader->IsError())
+				{
+					return ReadSize;
+				}
+			}
+			return 0;
+		}
+
+		virtual int32 GetFrameId() const override
+		{
+			return FrameId;
+		}
+
+		virtual const FString& GetFilename() const override
+		{
+			return Filename;
+		}
+
+		virtual FArchive* GetUnderlyingArchive() const override
+		{
+			return FileReader.Get();
+		}
+		//~ End ITmvMediaDecoderAccessUnit
+
+	private:
+		int32 FrameId = 0;
+		FString Filename;
+		int64 BaseOffset = 0;
+		int64 SampleSize = 0;
+		TUniquePtr<FArchive> FileReader;
+	};
+
 	/**
 	* Pool for worker thread context objects that can be reused from one "ReadFrame" to the next when available.
 	*/
@@ -218,14 +321,35 @@ namespace UE::ImgMedia
 		return TotalSize;
 	}
 	
-	void FillFrameInfo(const ITmvMediaDecoderFactory& InDecoderFactory, const TArray<FTmvMediaFrameMipInfo>& InMipInfos, FImgMediaFrameInfo& OutInfo)
+	void FillFrameInfo(
+		const ITmvMediaDecoderFactory& InDecoderFactory, 
+		const TArray<FTmvMediaFrameMipInfo>& InMipInfos, 
+		const FTmvMediaDemuxerTrackInfo* InVideoTrackInfo,
+		FImgMediaFrameInfo& OutInfo)
 	{
 		OutInfo.CompressionName = InDecoderFactory.GetName();	// todo: codec: APV, or one of the exr compressions.
 		OutInfo.Dim = FIntPoint(InMipInfos[0].Width, InMipInfos[0].Height);
 		OutInfo.NumMipLevels = InMipInfos.Num();
 		OutInfo.UncompressedSize = CalculateMemorySizeInBytes(InMipInfos);
 		OutInfo.FormatName = "TMV";	// todo: Could be EXR or APV.
-		OutInfo.FrameRate = GetDefault<UImgMediaSettings>()->DefaultFrameRate;
+		if (InVideoTrackInfo)
+		{
+			if (InVideoTrackInfo->ConstantSampleDuration > 0 && InVideoTrackInfo->Timescale > 0)
+			{
+				OutInfo.FrameRate = FFrameRate(InVideoTrackInfo->Timescale, InVideoTrackInfo->ConstantSampleDuration);
+			}
+			else if (InVideoTrackInfo->NumSamples > 0 && InVideoTrackInfo->Duration > 0)
+			{
+				// Fallback: average frame rate from total duration and sample count.
+				OutInfo.FrameRate = FFrameRate(
+					static_cast<int32>(static_cast<int64>(InVideoTrackInfo->NumSamples) * InVideoTrackInfo->Timescale / InVideoTrackInfo->Duration),
+					1);
+			}
+		}
+		else
+		{
+			OutInfo.FrameRate = GetDefault<UImgMediaSettings>()->DefaultFrameRate;
+		}
 		OutInfo.Srgb = true;	// todo: Transfer function and color space info.
 		OutInfo.NumChannels = InMipInfos[0].NumComponents;
 		OutInfo.NumBytesPerPixel = OutInfo.NumChannels * 2; // todo: This assumes 16 bits per channel
@@ -381,8 +505,76 @@ namespace UE::ImgMedia
 		// Empty array view.
 		return TConstArrayView<FTmvMediaShaderTileDesc>();
 	}
+
+	/**
+	 * Utility function to create a reader if the specified file is a supported container format.
+	 * @return Created Reader or null if not supported.
+	 */
+	TSharedPtr<FTmvImgMediaReader> CreateContainerReader(
+		const TSharedRef<FImgMediaLoader>& InLoader, 
+		const FString& InMediaFilePath, 
+		const FString& InExtension, 
+		ITmvMediaModule& InTmvModule)
+	{
+		// Check if this is a container format supported by a registered demuxer factory.
+		const TSharedPtr<ITmvMediaDemuxerFactory, ESPMode::ThreadSafe> DemuxerFactory = InTmvModule.FindDemuxerFactoryForExtension(InExtension);
+		if (!DemuxerFactory.IsValid())
+		{
+			return nullptr;
+		}
+		
+		// Create and open the demuxer.
+		const TSharedPtr<ITmvMediaDemuxer, ESPMode::ThreadSafe> Demuxer = DemuxerFactory->CreateDemuxer();
+		if (!Demuxer.IsValid())
+		{
+			return nullptr;
+		}
+
+		if (Demuxer->OpenFile(InMediaFilePath) != ETmvMediaContainerResult::Success)
+		{
+			UE_LOGF(LogImgMedia, Warning, "FTmvImgMediaReader::CreateReader: Failed to open container: %ls", *Demuxer->GetLastError());
+			return nullptr;
+		}
+
+		// Find the first video track.
+		int32 VideoTrackIndex = INDEX_NONE;
+		FTmvMediaDemuxerTrackInfo VideoTrackInfo;
+		for (int32 i = 0; i < Demuxer->GetTrackCount(); ++i)
+		{
+			FTmvMediaDemuxerTrackInfo Info;
+			if (Demuxer->GetTrackInfo(i, Info) == ETmvMediaContainerResult::Success && Info.TrackType == ETmvMediaTrackType::Video)
+			{
+				VideoTrackIndex = i;
+				VideoTrackInfo = Info;
+				break;
+			}
+		}
+
+		if (VideoTrackIndex == INDEX_NONE)
+		{
+			UE_LOGF(LogImgMedia, Warning, "FTmvImgMediaReader::CreateReader: No video track found in container \"%ls\".", *InMediaFilePath);
+			Demuxer->Close();
+			return nullptr;
+		}
+
+		// Find a decoder factory for the track's codec.
+		const FString CodecFormat = UE::TmvMedia::Utils::FourCCToString(VideoTrackInfo.SampleEntryFormat);
+
+		TSharedPtr<ITmvMediaDecoderFactory, ESPMode::ThreadSafe> DecoderFactory = InTmvModule.GetBestDecoderFactoryForFormat(CodecFormat, {});
+		if (!DecoderFactory.IsValid())
+		{
+			UE_LOGF(LogImgMedia, Warning, "FTmvImgMediaReader::CreateReader: No decoder factory for codec '%ls'.", *CodecFormat);
+			Demuxer->Close();
+			return nullptr;
+		}
+		return MakeShared<FTmvContainerImgMediaReader>(InLoader, DecoderFactory.ToSharedRef(), Demuxer, VideoTrackIndex, InMediaFilePath);
+	}
 }
 
+// -------------------------------------------------------------------------------------------------
+// FTmvImgMediaReader (abstract base)
+// -------------------------------------------------------------------------------------------------
+
 FTmvImgMediaReader::FTmvImgMediaReader(const TSharedRef<FImgMediaLoader>& InLoader, const TSharedRef<ITmvMediaDecoderFactory>& InDecoderFactory)
 	: LoaderWeak(InLoader)
 	, DecoderFactory(InDecoderFactory)
@@ -395,13 +587,13 @@ FTmvImgMediaReader::FTmvImgMediaReader(const TSharedRef<FImgMediaLoader>& InLoad
 
 FTmvImgMediaReader::~FTmvImgMediaReader() = default;
 
-TArray<FString> FTmvImgMediaReader::GetSupportedFileExtensions()
+TArray<FString> FTmvImgMediaReader::GetSupportedImageFileExtensions()
 {
 	TArray<FString> OutExtensions;
 
-	ITmvMediaModule& TmvModule = ITmvMediaModule::GetOrLoad();	
+	// File-sequence extensions from decoder factories.
 	TArray<TWeakPtr<ITmvMediaDecoderFactory, ESPMode::ThreadSafe>> DecoderFactories;
-	TmvModule.GetDecoderFactories(DecoderFactories);
+	ITmvMediaModule::GetOrLoad().GetDecoderFactories(DecoderFactories);
 	for (const TWeakPtr<ITmvMediaDecoderFactory>& FactoryWeak : DecoderFactories)
 	{
 		if (TSharedPtr<ITmvMediaDecoderFactory> Factory = FactoryWeak.Pin())
@@ -412,11 +604,28 @@ TArray<FString> FTmvImgMediaReader::GetSupportedFileExtensions()
 	return OutExtensions;
 }
 
-bool FTmvImgMediaReader::IsFileExtensionSupported(const FString& InExtension)
+TArray<FString> FTmvImgMediaReader::GetSupportedContainerFileExtensions()
 {
-	ITmvMediaModule& TmvModule = ITmvMediaModule::GetOrLoad();
+	TArray<FString> OutExtensions;
+
+	// Container extensions from demuxer factories.
+	TArray<TWeakPtr<ITmvMediaDemuxerFactory, ESPMode::ThreadSafe>> DemuxerFactories;
+	ITmvMediaModule::GetOrLoad().GetDemuxerFactories(DemuxerFactories);
+	for (const TWeakPtr<ITmvMediaDemuxerFactory>& FactoryWeak : DemuxerFactories)
+	{
+		if (TSharedPtr<ITmvMediaDemuxerFactory> Factory = FactoryWeak.Pin())
+		{
+			OutExtensions.Append(Factory->GetSupportedContainerFormats());
+		}
+	}
+	return OutExtensions;
+}
+
+bool FTmvImgMediaReader::IsImageFileExtensionSupported(const FString& InExtension)
+{
+	// Check decoder factories (file sequence).
 	TArray<TWeakPtr<ITmvMediaDecoderFactory, ESPMode::ThreadSafe>> DecoderFactories;
-	TmvModule.GetDecoderFactories(DecoderFactories);
+	ITmvMediaModule::GetOrLoad().GetDecoderFactories(DecoderFactories);
 	for (const TWeakPtr<ITmvMediaDecoderFactory>& FactoryWeak : DecoderFactories)
 	{
 		if (TSharedPtr<ITmvMediaDecoderFactory> Factory = FactoryWeak.Pin())
@@ -430,47 +639,61 @@ bool FTmvImgMediaReader::IsFileExtensionSupported(const FString& InExtension)
 	return false;
 }
 
-/** Creates a reader of the appropriate type depending on the context (input file, settings, etc). */
-TSharedPtr<IImgMediaReader> FTmvImgMediaReader::CreateReader(const TSharedRef<FImgMediaLoader>& InLoader, const FString& InFirstImageInSequencePath)
+bool FTmvImgMediaReader::IsContainerFileExtensionSupported(const FString& InExtension)
 {
-	const FString Extension = FPaths::GetExtension(InFirstImageInSequencePath);
+	// Check demuxer factories (container).
+	return ITmvMediaModule::GetOrLoad().FindDemuxerFactoryForExtension(InExtension).IsValid();
+}
+
+TSharedPtr<FTmvImgMediaReader> FTmvImgMediaReader::CreateReader(const TSharedRef<FImgMediaLoader>& InLoader, const FString& InMediaFilePath)
+{
+	const FString Extension = FPaths::GetExtension(InMediaFilePath);
 
 	ITmvMediaModule& TmvModule = ITmvMediaModule::GetOrLoad();
+
+	TSharedPtr<FTmvImgMediaReader> ContainerReader = UE::ImgMedia::CreateContainerReader(InLoader, InMediaFilePath, Extension, TmvModule);
+	if (ContainerReader.IsValid())
+	{
+		return ContainerReader;
+	}
+
+	// Fall back to file-sequence reader.
 	TSharedPtr<ITmvMediaDecoderFactory, ESPMode::ThreadSafe> Factory = TmvModule.GetBestDecoderFactoryForFormat(Extension, {});
 	if (Factory.IsValid())
 	{
-		return MakeShared<FTmvImgMediaReader>(InLoader, Factory.ToSharedRef());
+		return MakeShared<FTmvFileSequenceImgMediaReader>(InLoader, Factory.ToSharedRef());
 	}
-	return nullptr;	
+	return nullptr;
 }
 
 bool FTmvImgMediaReader::GetFrameInfo(const FString& InImagePath, FImgMediaFrameInfo& OutInfo)
 {
 	check(ParserPool);
 
-	using namespace UE::ImgMedia;
-	FFileAccessUnit FileAccessUnit(InImagePath, 0);
-	if (!FileAccessUnit.IsValid())
+	TUniquePtr<ITmvMediaDecoderAccessUnit> AccessUnit = CreateAccessUnit(0);
+	if (!AccessUnit)
 	{
-		UE_LOGF(LogImgMedia, Verbose, "FTmvImgMediaReader: Failed to open TMV access unit file container %ls", *InImagePath);
+		UE_LOGF(LogImgMedia, Verbose, "FTmvImgMediaReader: Failed to create access unit for frame info");
 		return false;
 	}
 
+	using namespace UE::ImgMedia;
+
 	TArray<FTmvMediaFrameMipInfo> MipInfos;
 	{
 		FTmvParserPool::FHandle ParserHandle = ParserPool->AcquireParser(*DecoderFactory);
-		if (ParserHandle->ParseMipInfos(FileAccessUnit, MipInfos) != ETmvMediaDecoderResult::Success)
+		if (ParserHandle->ParseMipInfos(*AccessUnit, MipInfos) != ETmvMediaDecoderResult::Success)
 		{
 			return false;
 		}
-
 		if (MipInfos.Num() == 0)
 		{
 			return false;
 		}
 	}
 
-	FillFrameInfo(*DecoderFactory, MipInfos, OutInfo);
+	FillFrameInfo(*DecoderFactory, MipInfos, GetVideoTrackInfo(), OutInfo);
+	
 	return true;
 }
 
@@ -488,13 +711,12 @@ bool FTmvImgMediaReader::ReadFrame(int32 InFrameId, const TMap<int32, FImgMediaT
 		return false;
 	}
 
-	const FString& ImagePath = Loader->GetImagePath(InFrameId, 0);
 	using namespace UE::ImgMedia;
 
-	FFileAccessUnit FileAccessUnit(ImagePath, InFrameId);
-	if (!FileAccessUnit.IsValid())
+	TUniquePtr<ITmvMediaDecoderAccessUnit> AccessUnit = CreateAccessUnit(InFrameId);
+	if (!AccessUnit)
 	{
-		UE_LOGF(LogImgMedia, Warning, "FTmvImgMediaReader: Failed to open TMV access unit file container %ls", *ImagePath);
+		UE_LOGF(LogImgMedia, Warning, "FTmvImgMediaReader: Failed to create access unit for frame %d", InFrameId);
 		return false;
 	}
 
@@ -502,7 +724,7 @@ bool FTmvImgMediaReader::ReadFrame(int32 InFrameId, const TMap<int32, FImgMediaT
 	{
 		TRACE_CPUPROFILER_EVENT_SCOPE(TmvImgMedia.ParseAU);
 		FTmvParserPool::FHandle ParserHandle = ParserPool->AcquireParser(*DecoderFactory);
-		if (ParserHandle->ParseMipInfos(FileAccessUnit, MipInfos) != ETmvMediaDecoderResult::Success)
+		if (ParserHandle->ParseMipInfos(*AccessUnit, MipInfos) != ETmvMediaDecoderResult::Success)
 		{
 			return false;
 		}
@@ -513,7 +735,7 @@ bool FTmvImgMediaReader::ReadFrame(int32 InFrameId, const TMap<int32, FImgMediaT
 	}
 
 	FImgMediaFrameInfo FrameInfo;
-	FillFrameInfo(*DecoderFactory, MipInfos, FrameInfo);
+	FillFrameInfo(*DecoderFactory, MipInfos, GetVideoTrackInfo(), FrameInfo);
 	OutFrame->SetInfo(FrameInfo);
 
 	TSharedPtr<FTmvMediaTextureSampleConverter> SampleConverter = OutFrame->GetOrCreateSampleConverter<FTmvMediaTextureSampleConverter>();
@@ -530,9 +752,9 @@ bool FTmvImgMediaReader::ReadFrame(int32 InFrameId, const TMap<int32, FImgMediaT
 	ConverterParams.TileDimWithBorders = FrameInfo.TileDimensions + FrameInfo.TileBorder * 2;
 	ConverterParams.SourceColorSettings = Loader->GetSourceColorSettings();
 	const int32 NumMipLevels = Loader->GetNumMipLevels();
-	
+
 	// Force mip level to be upscaled to all higher quality mips.
-	TMap<int32, FImgMediaTileSelection> InMipTilesCopy = InMipTiles;
+	TMap<int32, FImgMediaTileSelection> MipTilesCopy = InMipTiles;
 	const int32 MipToUpscale = FMath::Clamp(Loader->GetMinimumLevelToUpscale(), -1, NumMipLevels - 1);
 
 	if (NumMipLevels > 1 && MipToUpscale >= 0)
@@ -540,13 +762,13 @@ bool FTmvImgMediaReader::ReadFrame(int32 InFrameId, const TMap<int32, FImgMediaT
 		ConverterParams.UpscaleMip = MipToUpscale;
 
 		FImgMediaTileSelection FullSelection = FImgMediaTileSelection::CreateForTargetMipLevel(ConverterParams.FullResolution, FrameInfo.TileDimensions, MipToUpscale, true);
-		if (InMipTilesCopy.Contains(MipToUpscale))
+		if (MipTilesCopy.Contains(MipToUpscale))
 		{
-			InMipTilesCopy[MipToUpscale] = FullSelection;
+			MipTilesCopy[MipToUpscale] = FullSelection;
 		}
 		else
 		{
-			InMipTilesCopy.Add(MipToUpscale, FullSelection);
+			MipTilesCopy.Add(MipToUpscale, FullSelection);
 		}
 	}
 
@@ -558,9 +780,9 @@ bool FTmvImgMediaReader::ReadFrame(int32 InFrameId, const TMap<int32, FImgMediaT
 	}
 
 	TArray<FTmvMediaDecoderMipRequest> MipRequests;
-	MipRequests.Reserve(InMipTilesCopy.Num());
+	MipRequests.Reserve(MipTilesCopy.Num());
 
-	for (const TPair<int32, FImgMediaTileSelection>& TilesPerMip : InMipTilesCopy)
+	for (const TPair<int32, FImgMediaTileSelection>& TilesPerMip : MipTilesCopy)
 	{
 		const int32 CurrentMipLevel = TilesPerMip.Key;
 		const FImgMediaTileSelection& CurrentTileSelection = TilesPerMip.Value;
@@ -595,7 +817,7 @@ bool FTmvImgMediaReader::ReadFrame(int32 InFrameId, const TMap<int32, FImgMediaT
 		switch (Result)
 		{
 		case ETmvSetupResult::Success:
-				break;
+			break;
 		case ETmvSetupResult::Fail:
 			return false;	// Abort further reading (for now)
 		case ETmvSetupResult::Skipped:
@@ -608,7 +830,7 @@ bool FTmvImgMediaReader::ReadFrame(int32 InFrameId, const TMap<int32, FImgMediaT
 		}
 	}
 
-	ETmvMediaDecoderResult DecodeResult = DecoderHandle->Decode(FileAccessUnit, MipRequests);
+	ETmvMediaDecoderResult DecodeResult = DecoderHandle->Decode(*AccessUnit, MipRequests);
 	if (DecodeResult != ETmvMediaDecoderResult::Success)
 	{
 		UE_LOGF(LogImgMedia, Verbose, "TmvImgMediaReader: Failed tile/mip decoding request for frame %d.", InFrameId);
@@ -620,7 +842,7 @@ bool FTmvImgMediaReader::ReadFrame(int32 InFrameId, const TMap<int32, FImgMediaT
 		int32 MipLevel = MipRequest.MipInfo.MipLevel;
 		if (MipRequest.OutResult == ETmvMediaDecoderResult::Success)
 		{
-			OutFrame->MipTilesPresent.Include(MipLevel, InMipTilesCopy[MipLevel]);
+			OutFrame->MipTilesPresent.Include(MipLevel, MipTilesCopy[MipLevel]);
 			OutFrame->NumTilesRead += MipRequest.OutNumTilesDecoded;
 
 			// Upload newly decoded tiles to gpu.
@@ -636,12 +858,93 @@ bool FTmvImgMediaReader::ReadFrame(int32 InFrameId, const TMap<int32, FImgMediaT
 	// Create viewport(s) with all mip/tiles present in the frame.
 	UpdateTmvConverterViewports(ConverterParams, SampleConverter, OutFrame, ConverterParams.Viewports);
 
-	// The final output frame after the sample converter is RGBA. 
-	OutFrame->Format = EMediaTextureSampleFormat::FloatRGBA;
-	OutFrame->Stride = ConverterParams.FullResolution.X * 4 * 2;	// Todo: verify if this is needed.
+	// The final output frame after the sample converter is RGBA.
+	OutFrame->Format = EMediaTextureSampleFormat::FloatRGBA;	// todo: investigate using FloatRGB (half memory footprint).
+	OutFrame->Stride = ConverterParams.FullResolution.X * 4 * 2;
 
 	SampleConverter->SetParams_ThreadSafe(ConverterParams);
 
 	UE_LOGF(LogImgMedia, Verbose, "Reader %p: Read Pixels Complete. %i", this, InFrameId);
 	return true;
-}
\ No newline at end of file
+}
+
+// -------------------------------------------------------------------------------------------------
+// FTmvFileSequenceImgMediaReader
+// -------------------------------------------------------------------------------------------------
+
+FTmvFileSequenceImgMediaReader::FTmvFileSequenceImgMediaReader(const TSharedRef<FImgMediaLoader>& InLoader, const TSharedRef<ITmvMediaDecoderFactory>& InDecoderFactory)
+	: FTmvImgMediaReader(InLoader, InDecoderFactory)
+{
+}
+
+TUniquePtr<ITmvMediaDecoderAccessUnit> FTmvFileSequenceImgMediaReader::CreateAccessUnit(int32 InFrameId)
+{
+	const TSharedPtr<FImgMediaLoader, ESPMode::ThreadSafe> Loader = LoaderWeak.Pin();
+	if (!Loader.IsValid())
+	{
+		return nullptr;
+	}
+
+	const FString& ImagePath = Loader->GetImagePath(InFrameId, 0);
+	auto AccessUnit = MakeUnique<UE::ImgMedia::FFileAccessUnit>(ImagePath, InFrameId);
+	if (!AccessUnit->IsValid())
+	{
+		UE_LOGF(LogImgMedia, Warning, "FTmvFileSequenceImgMediaReader: Failed to open file %ls", *ImagePath);
+		return nullptr;
+	}
+	return AccessUnit;
+}
+
+// -------------------------------------------------------------------------------------------------
+// FTmvContainerImgMediaReader
+// -------------------------------------------------------------------------------------------------
+
+FTmvContainerImgMediaReader::FTmvContainerImgMediaReader(
+	const TSharedRef<FImgMediaLoader>& InLoader,
+	const TSharedRef<ITmvMediaDecoderFactory>& InDecoderFactory,
+	const TSharedPtr<ITmvMediaDemuxer, ESPMode::ThreadSafe>& InDemuxer,
+	int32 InTrackIndex,
+	const FString& InContainerFilePath)
+	: FTmvImgMediaReader(InLoader, InDecoderFactory)
+	, Demuxer(InDemuxer)
+	, DemuxerTrackIndex(InTrackIndex)
+	, ContainerFilePath(InContainerFilePath)
+{
+	Demuxer->GetTrackInfo(DemuxerTrackIndex, VideoTrackInfo);
+}
+
+FTmvContainerImgMediaReader::~FTmvContainerImgMediaReader()
+{
+	if (Demuxer.IsValid())
+	{
+		Demuxer->Close();
+	}
+}
+
+TUniquePtr<ITmvMediaDecoderAccessUnit> FTmvContainerImgMediaReader::CreateAccessUnit(int32 InFrameId)
+{
+	using namespace UE::ImgMedia;
+
+	FTmvMediaDemuxerSample SampleInfo;
+	{
+		FScopeLock Lock(&DemuxerCS);
+		if (Demuxer->SeekToSample(DemuxerTrackIndex, static_cast<uint32>(InFrameId)) != ETmvMediaContainerResult::Success)
+		{
+			UE_LOGF(LogImgMedia, Warning, "FTmvContainerImgMediaReader: Failed to seek to sample %d", InFrameId);
+			return nullptr;
+		}
+		if (Demuxer->ReadSampleInfo(DemuxerTrackIndex, SampleInfo) != ETmvMediaContainerResult::Success)
+		{
+			UE_LOGF(LogImgMedia, Warning, "FTmvContainerImgMediaReader: Failed to read sample info for frame %d", InFrameId);
+			return nullptr;
+		}
+	}
+
+	auto AccessUnit = MakeUnique<FContainerFileAccessUnit>(ContainerFilePath, SampleInfo.FileOffset, SampleInfo.SampleSize, InFrameId);
+	if (!AccessUnit->IsValid())
+	{
+		UE_LOGF(LogImgMedia, Warning, "FTmvContainerImgMediaReader: Failed to open container access unit for frame %d", InFrameId);
+		return nullptr;
+	}
+	return AccessUnit;
+}
diff --git a/Engine/Plugins/Media/ImgMedia/Source/ImgMedia/Private/Readers/TmvImgMediaReader.h b/Engine/Plugins/Media/ImgMedia/Source/ImgMedia/Private/Readers/TmvImgMediaReader.h
--- a/Engine/Plugins/Media/ImgMedia/Source/ImgMedia/Private/Readers/TmvImgMediaReader.h
+++ b/Engine/Plugins/Media/ImgMedia/Source/ImgMedia/Private/Readers/TmvImgMediaReader.h
@@ -3,11 +3,15 @@
 #pragma once
 
 #include "Containers/ContainersFwd.h"
+#include "Decoder/ITmvMediaDemuxer.h"
 #include "IImgMediaReader.h"
 
 class FImgMediaLoader;
 class FTmvMediaFrameMipBufferPool;
+class ITmvMediaDecoderAccessUnit;
 class ITmvMediaDecoderFactory;
+class ITmvMediaDemuxer;
+class ITmvMediaDemuxerFactory;
 
 namespace UE::ImgMedia
 {
@@ -16,39 +20,63 @@ namespace UE::ImgMedia
 }
 
 /**
- * Implements an img reader wrapper for Tmv Decoders.
+ * Abstract base class for TMV image media readers.
+ * Provides the common decode pipeline (parse, decode, compose).
+ * Subclasses provide the access unit for each frame.
  */
-class FTmvImgMediaReader final : public IImgMediaReader
+class FTmvImgMediaReader : public IImgMediaReader
 {
 public:
+	FTmvImgMediaReader(const TSharedRef<FImgMediaLoader>& InLoader, const TSharedRef<ITmvMediaDecoderFactory>& InDecoderFactory);
+	virtual ~FTmvImgMediaReader() override;
+
+	/** Get the list of supported image file extensions for TMV decoders. */
+	static TArray<FString> GetSupportedImageFileExtensions();
+	
+	/** Get the list of supported container file extensions for TMV demuxers. */
+	static TArray<FString> GetSupportedContainerFileExtensions();
 
 	/**
-	 * Create and initialize a new instance.
-	 *
-	 * @param InLoader Parent loader.
-	 * @param InDecoderFactory Decoder factory.
+	 * Returns true if the given image file extension is supported by a Tmv Reader. 
 	 */
-	FTmvImgMediaReader(const TSharedRef<FImgMediaLoader>& InLoader, const TSharedRef<ITmvMediaDecoderFactory>& InDecoderFactory);
+	static bool IsImageFileExtensionSupported(const FString& InExtension);
 
-	virtual ~FTmvImgMediaReader() override;
-
-	/** Get the list of supported file extensions for TMV decoders. */
-	static TArray<FString> GetSupportedFileExtensions();
+	/**
+	 * Returns true if the given container file extension is supported by a Tmv Reader. 
+	 */
+	static bool IsContainerFileExtensionSupported(const FString& InExtension);
 
-	static bool IsFileExtensionSupported(const FString& InExtension);
+	/** 
+	 * Creates a reader of the appropriate type depending on the context (input file, settings, etc).
+	 * @param InLoader Parent Img Loader
+	 * @param InMediaFilePath Full path of the media file, either the first image of the sequence or the container.
+	 */
+	static TSharedPtr<FTmvImgMediaReader> CreateReader(const TSharedRef<FImgMediaLoader>& InLoader, const FString& InMediaFilePath);
 	
-	/** Creates a reader of the appropriate type depending on the context (input file, settings, etc). */
-	static TSharedPtr<IImgMediaReader> CreateReader(const TSharedRef<FImgMediaLoader>& InLoader, const FString& InFirstImageInSequencePath);
+	/** 
+	 * Returns the video track info if present, nullptr otherwise.
+	 */
+	virtual const FTmvMediaDemuxerTrackInfo* GetVideoTrackInfo() const
+	{
+		return nullptr;
+	}
 
 public:
 	//~ Begin IImgMediaReader
 	virtual bool GetFrameInfo(const FString& InImagePath, FImgMediaFrameInfo& OutInfo) override;
 	virtual bool ReadFrame(int32 InFrameId, const TMap<int32, FImgMediaTileSelection>& InMipTiles, TSharedPtr<FImgMediaFrame, ESPMode::ThreadSafe> OutFrame) override;
-	virtual void CancelFrame(int32 InFrameNumber) override {};
-	virtual void UncancelFrame(int32 InFrameNumber) override {};
+	virtual void CancelFrame(int32 InFrameNumber) override {}
+	virtual void UncancelFrame(int32 InFrameNumber) override {}
 	//~ End IImgMediaReader
 
-private:
+protected:
+	/**
+	 * Create an access unit for the given frame.
+	 * @param InFrameId Frame index.
+	 * @return A valid access unit, or null on failure.
+	 */
+	virtual TUniquePtr<ITmvMediaDecoderAccessUnit> CreateAccessUnit(int32 InFrameId) = 0;
+
 	/** Parent loader. */
 	TWeakPtr<FImgMediaLoader> LoaderWeak;
 
@@ -63,4 +91,63 @@ private:
 
 	/** Texture Sample Buffer Pool for this image reader. */
 	TSharedPtr<FTmvMediaFrameMipBufferPool> FrameMipBufferPool;
-};
\ No newline at end of file
+};
+
+/**
+ * TMV reader for file sequences: one file per frame.
+ */
+class FTmvFileSequenceImgMediaReader final : public FTmvImgMediaReader
+{
+public:
+	FTmvFileSequenceImgMediaReader(const TSharedRef<FImgMediaLoader>& InLoader, const TSharedRef<ITmvMediaDecoderFactory>& InDecoderFactory);
+
+protected:
+	//~ Begin FTmvImgMediaReader
+	virtual TUniquePtr<ITmvMediaDecoderAccessUnit> CreateAccessUnit(int32 InFrameId) override;
+	//~ End FTmvImgMediaReader
+};
+
+/**
+ * TMV reader for container files: reads frames from a .tmv container via ITmvMediaDemuxer.
+ * Uses ReadSampleInfo (metadata only) + direct file reads (no intermediate memory buffer).
+ */
+class FTmvContainerImgMediaReader final : public FTmvImgMediaReader
+{
+public:
+	FTmvContainerImgMediaReader(
+		const TSharedRef<FImgMediaLoader>& InLoader,
+		const TSharedRef<ITmvMediaDecoderFactory>& InDecoderFactory,
+		const TSharedPtr<ITmvMediaDemuxer, ESPMode::ThreadSafe>& InDemuxer,
+		int32 InTrackIndex,
+		const FString& InContainerFilePath);
+
+	virtual ~FTmvContainerImgMediaReader() override;
+
+	//~ Begin FTmvImgMediaReader
+	virtual const FTmvMediaDemuxerTrackInfo* GetVideoTrackInfo() const override
+	{
+		return (VideoTrackInfo.TrackType == ETmvMediaTrackType::Video) ? &VideoTrackInfo : nullptr;
+	}
+	//~ End FTmvImgMediaReader
+
+protected:
+	//~ Begin FTmvImgMediaReader
+	virtual TUniquePtr<ITmvMediaDecoderAccessUnit> CreateAccessUnit(int32 InFrameId) override;
+	//~ End FTmvImgMediaReader
+
+private:
+	/** Demuxer for reading sample info from the container. */
+	TSharedPtr<ITmvMediaDemuxer, ESPMode::ThreadSafe> Demuxer;
+
+	/** Serializes SeekToSample + ReadSampleInfo on the demuxer. */
+	FCriticalSection DemuxerCS;
+
+	/** Video track index in the demuxer. */
+	int32 DemuxerTrackIndex = 0;
+
+	/** Cached video track info from the demuxer. */
+	FTmvMediaDemuxerTrackInfo VideoTrackInfo;
+
+	/** Path to the container file (for creating per-thread file access units). */
+	FString ContainerFilePath;
+};
diff --git a/Engine/Plugins/Media/ImgMedia/Source/ImgMedia/Public/IImgMediaModule.h b/Engine/Plugins/Media/ImgMedia/Source/ImgMedia/Public/IImgMediaModule.h
--- a/Engine/Plugins/Media/ImgMedia/Source/ImgMedia/Public/IImgMediaModule.h
+++ b/Engine/Plugins/Media/ImgMedia/Source/ImgMedia/Public/IImgMediaModule.h
@@ -12,7 +12,6 @@
 #define UE_API IMGMEDIA_API
 
 class FImgMediaGlobalCache;
-struct FImgMediaMipMapCameraInfo;
 class FImgMediaPlayer;
 class FImgMediaSceneViewExtension;
 class IMediaEventSink;
@@ -56,7 +55,7 @@ public:
 	 */
 	virtual TSharedPtr<IMediaPlayer, ESPMode::ThreadSafe> CreatePlayer(IMediaEventSink& EventHandler) = 0;
 
-	/** Get scene view extension cahcing view information (for mip/tile calculations). */
+	/** Get scene view extension caching view information (for mip/tile calculations). */
 	virtual const TSharedPtr<FImgMediaSceneViewExtension, ESPMode::ThreadSafe>& GetSceneViewExtension() const = 0;
 public:
 
@@ -68,7 +67,7 @@ public:
 
 	/**
 	 * Call this to get the global cache.
-	 * 
+	 *
 	 * @return Global cache.
 	 */
 	static FImgMediaGlobalCache* GetGlobalCache() { return GlobalCache.Get(); }
diff --git a/Engine/Plugins/Media/ImgMedia/Source/ImgMediaEditor/ImgMediaEditor.Build.cs b/Engine/Plugins/Media/ImgMedia/Source/ImgMediaEditor/ImgMediaEditor.Build.cs
--- a/Engine/Plugins/Media/ImgMedia/Source/ImgMediaEditor/ImgMediaEditor.Build.cs
+++ b/Engine/Plugins/Media/ImgMedia/Source/ImgMediaEditor/ImgMediaEditor.Build.cs
@@ -23,6 +23,7 @@ namespace UnrealBuildTool.Rules
 					"RenderCore",
 					"Slate",
 					"SlateCore",
+					"TmvMedia",
 					"ToolMenus",
 					"UnrealEd",
 					"WorkspaceMenuStructure",
diff --git a/Engine/Plugins/Media/ImgMedia/Source/ImgMediaEditor/Private/Factories/TmvFileMediaSourceFactory.cpp b/Engine/Plugins/Media/ImgMedia/Source/ImgMediaEditor/Private/Factories/TmvFileMediaSourceFactory.cpp
new file mode 100644
--- /dev/null
+++ b/Engine/Plugins/Media/ImgMedia/Source/ImgMediaEditor/Private/Factories/TmvFileMediaSourceFactory.cpp
@@ -0,0 +1,87 @@
+// Copyright Epic Games, Inc. All Rights Reserved.
+
+#include "TmvFileMediaSourceFactory.h"
+
+#include "Misc/CoreDelegates.h"
+#include "Misc/Paths.h"
+#include "FileMediaSource.h"
+#include "ITmvMediaModule.h"
+#include "Containers/UnrealString.h"
+#include "Decoder/ITmvMediaDemuxerFactory.h"
+#include "UObject/UObjectGlobals.h"
+
+#include UE_INLINE_GENERATED_CPP_BY_NAME(TmvFileMediaSourceFactory)
+
+
+/* UTmvFileMediaSourceFactory structors
+ *****************************************************************************/
+
+UTmvFileMediaSourceFactory::UTmvFileMediaSourceFactory(const FObjectInitializer& ObjectInitializer)
+	: Super(ObjectInitializer)
+{
+	SupportedClass = UFileMediaSource::StaticClass();
+	bEditorImport = true;
+
+	// Try now (may be empty if demuxer factories haven't been registered yet).
+	InitializeFormats();
+
+	// For the CDO, also register for PostEngineInit to retry once all modules are loaded.
+	if (HasAnyFlags(RF_ClassDefaultObject))
+	{
+		FCoreDelegates::OnAllModuleLoadingPhasesComplete.AddWeakLambda(this, [this]()
+		{
+			InitializeFormats();
+		});
+	}
+}
+
+void UTmvFileMediaSourceFactory::InitializeFormats()
+{
+	const ITmvMediaModule* TmvMediaModule = ITmvMediaModule::Get();
+	if (!TmvMediaModule)
+	{
+		return;
+	}
+
+	Formats.Reset();
+
+	TArray<TWeakPtr<ITmvMediaDemuxerFactory, ESPMode::ThreadSafe>> DemuxerFactories;
+	TmvMediaModule->GetDemuxerFactories(DemuxerFactories);
+	for (const TWeakPtr<ITmvMediaDemuxerFactory, ESPMode::ThreadSafe>& DemuxerFactoryWeak : DemuxerFactories)
+	{
+		if (const TSharedPtr<ITmvMediaDemuxerFactory> DemuxerFactory = DemuxerFactoryWeak.Pin())
+		{
+			for (const FString& Extension : DemuxerFactory->GetSupportedContainerFormats())
+			{
+				Formats.Add(FString::Printf(TEXT("%s;%s"), *Extension, *DemuxerFactory->GetDisplayName().ToString()));
+			}
+		}
+	}
+}
+
+
+/* UFactory overrides
+ *****************************************************************************/
+
+bool UTmvFileMediaSourceFactory::FactoryCanImport(const FString& InFilename)
+{
+	if (const ITmvMediaModule* TmvMediaModule = ITmvMediaModule::Get())
+	{
+		const FString Extension = FPaths::GetExtension(InFilename);
+		if (TmvMediaModule->FindDemuxerFactoryForExtension(Extension).IsValid())
+		{
+			return true;
+		}
+	}
+	return false;
+}
+
+
+UObject* UTmvFileMediaSourceFactory::FactoryCreateFile(UClass* InClass, UObject* InParent, FName InName, EObjectFlags InFlags, const FString& InFilename, const TCHAR* InParms, FFeedbackContext* InWarn, bool& bOutOperationCanceled)
+{
+	UFileMediaSource* MediaSource = NewObject<UFileMediaSource>(InParent, InClass, InName, InFlags);
+	MediaSource->SetFilePath(CurrentFilename);
+
+	return MediaSource;
+}
+
diff --git a/Engine/Plugins/Media/ImgMedia/Source/ImgMediaEditor/Private/Factories/TmvFileMediaSourceFactory.h b/Engine/Plugins/Media/ImgMedia/Source/ImgMediaEditor/Private/Factories/TmvFileMediaSourceFactory.h
new file mode 100644
--- /dev/null
+++ b/Engine/Plugins/Media/ImgMedia/Source/ImgMediaEditor/Private/Factories/TmvFileMediaSourceFactory.h
@@ -0,0 +1,27 @@
+// Copyright Epic Games, Inc. All Rights Reserved.
+
+#pragma once
+
+#include "UObject/ObjectMacros.h"
+#include "Factories/Factory.h"
+
+#include "TmvFileMediaSourceFactory.generated.h"
+
+/**
+ * Implements a factory for importing Tmv container files as UFileMediaSource objects.
+ */
+UCLASS(hidecategories=Object)
+class UTmvFileMediaSourceFactory : public UFactory
+{
+	GENERATED_UCLASS_BODY()
+
+public:
+	//~ Begin UFactory
+	virtual bool FactoryCanImport(const FString& InFilename) override;
+	virtual UObject* FactoryCreateFile(UClass* InClass, UObject* InParent, FName InName, EObjectFlags InFlags, const FString& InFilename, const TCHAR* InParms, FFeedbackContext* InWarn, bool& bOutOperationCanceled) override;
+	//~ End UFactory
+
+private:
+	/** Populate the Formats array from registered demuxer factories. Safe to call multiple times. */
+	void InitializeFormats();
+};
diff --git a/Engine/Plugins/Media/ImgMedia/Source/ImgMediaFactory/ImgMediaFactory.Build.cs b/Engine/Plugins/Media/ImgMedia/Source/ImgMediaFactory/ImgMediaFactory.Build.cs
--- a/Engine/Plugins/Media/ImgMedia/Source/ImgMediaFactory/ImgMediaFactory.Build.cs
+++ b/Engine/Plugins/Media/ImgMedia/Source/ImgMediaFactory/ImgMediaFactory.Build.cs
@@ -14,7 +14,8 @@ namespace UnrealBuildTool.Rules
 			PrivateIncludePathModuleNames.AddRange(
 				new string[] {
 					"Media",
-					"ImgMedia"
+					"ImgMedia",
+					"TmvMedia"
 				});
 
 			PublicDependencyModuleNames.AddRange(
@@ -33,6 +34,7 @@ namespace UnrealBuildTool.Rules
 			if (Target.Platform.IsInGroup(UnrealPlatformGroup.Windows))
 			{
 				DynamicallyLoadedModuleNames.Add("ImgMedia");
+				DynamicallyLoadedModuleNames.Add("TmvMedia");
 			}
 		}
 	}
diff --git a/Engine/Plugins/Media/ImgMedia/Source/ImgMediaFactory/Private/ImgMediaFactoryModule.cpp b/Engine/Plugins/Media/ImgMedia/Source/ImgMediaFactory/Private/ImgMediaFactoryModule.cpp
--- a/Engine/Plugins/Media/ImgMedia/Source/ImgMediaFactory/Private/ImgMediaFactoryModule.cpp
+++ b/Engine/Plugins/Media/ImgMedia/Source/ImgMediaFactory/Private/ImgMediaFactoryModule.cpp
@@ -4,9 +4,13 @@
 
 #include "Containers/Array.h"
 #include "Containers/UnrealString.h"
+#include "Decoder/ITmvMediaDemuxerFactory.h"
+#include "IImgMediaModule.h"
 #include "IMediaModule.h"
 #include "IMediaOptions.h"
 #include "IMediaPlayerFactory.h"
+#include "ITmvMediaModule.h"
+#include "ImgMediaSettings.h"
 #include "Internationalization/Internationalization.h"
 #include "Misc/Paths.h"
 #include "Modules/ModuleInterface.h"
@@ -20,11 +24,6 @@
 	#include "UObject/WeakObjectPtr.h"
 #endif
 
-#include "ImgMediaSettings.h"
-
-#include "IImgMediaModule.h"
-
-
 DEFINE_LOG_CATEGORY(LogImgMediaFactory);
 
 #define LOCTEXT_NAMESPACE "FImgMediaFactoryModule"
@@ -68,6 +67,19 @@ public:
 			return 0;
 		}
 
+		// Check for supported container files via "file" scheme (e.g. file://path/to/video.tmv).
+		if (Scheme == TEXT("file"))
+		{
+			if (const ITmvMediaModule* TmvModule = FModuleManager::LoadModulePtr<ITmvMediaModule>("TmvMedia"))
+			{
+				const FString Extension = FPaths::GetExtension(Location);
+				if (TmvModule->FindDemuxerFactoryForExtension(Extension).IsValid())
+				{
+					return 100;
+				}
+			}
+		}
+
 		if (!SupportedUriSchemes.Contains(Scheme))
 		{
 			if (OutErrors != nullptr)
diff --git a/Engine/Plugins/Media/ImgMedia/Source/OpenExrWrapper/Private/OpenExrWrapperTmvEncoderOptions.h b/Engine/Plugins/Media/ImgMedia/Source/OpenExrWrapper/Private/OpenExrWrapperTmvEncoderOptions.h
--- a/Engine/Plugins/Media/ImgMedia/Source/OpenExrWrapper/Private/OpenExrWrapperTmvEncoderOptions.h
+++ b/Engine/Plugins/Media/ImgMedia/Source/OpenExrWrapper/Private/OpenExrWrapperTmvEncoderOptions.h
@@ -11,7 +11,7 @@
 /**
  * OpenExr Tmv Encoder options.
  */
-USTRUCT(BlueprintType, DisplayName = "OpenExr Tmv Encoder Options")
+USTRUCT(DisplayName = "OpenExr Tmv Encoder Options")
 struct FOpenExrWrapperTmvEncoderOptions : public FTmvMediaEncoderOptions
 {
 	GENERATED_BODY()
diff --git a/Engine/Plugins/Media/TmvMedia/Source/ApvMedia/Private/ApvMediaTmvEncoderFactory.cpp b/Engine/Plugins/Media/TmvMedia/Source/ApvMedia/Private/ApvMediaTmvEncoderFactory.cpp
--- a/Engine/Plugins/Media/TmvMedia/Source/ApvMedia/Private/ApvMediaTmvEncoderFactory.cpp
+++ b/Engine/Plugins/Media/TmvMedia/Source/ApvMedia/Private/ApvMediaTmvEncoderFactory.cpp
@@ -6,6 +6,7 @@
 #include "ApvMediaTmvEncoder.h"
 #include "ApvMediaTmvEncoderOptions.h"
 #include "StructUtils/InstancedStruct.h"
+#include "Utils/TmvMediaUtils.h"
 
 #define LOCTEXT_NAMESPACE "ApvMediaTmvEncoderFactory"
 
@@ -38,4 +39,10 @@ TSharedPtr<ITmvMediaEncoder, ESPMode::ThreadSafe> FApvMediaTmvEncoderFactory::Cr
 	return nullptr;
 }
 
-#undef LOCTEXT_NAMESPACE
\ No newline at end of file
+uint32 FApvMediaTmvEncoderOptions::GetCodecFourCC() const
+{
+	// FOURCC 'apv1'
+	return UE::TmvMedia::Utils::MakeFourCC('a', 'p', 'v', '1');
+}
+
+#undef LOCTEXT_NAMESPACE
diff --git a/Engine/Plugins/Media/TmvMedia/Source/ApvMedia/Private/ApvMediaTmvEncoderOptions.h b/Engine/Plugins/Media/TmvMedia/Source/ApvMedia/Private/ApvMediaTmvEncoderOptions.h
--- a/Engine/Plugins/Media/TmvMedia/Source/ApvMedia/Private/ApvMediaTmvEncoderOptions.h
+++ b/Engine/Plugins/Media/TmvMedia/Source/ApvMedia/Private/ApvMediaTmvEncoderOptions.h
@@ -273,7 +273,7 @@ namespace UE::ApvMedia
 /**
  * OpenApv Tmv Encoder options.
  */
-USTRUCT(BlueprintType, DisplayName = "OpenApv Tmv Encoder Options")
+USTRUCT(DisplayName = "OpenApv Tmv Encoder Options")
 struct FApvMediaTmvEncoderOptions : public FTmvMediaEncoderOptions
 {
 	GENERATED_BODY()
@@ -293,6 +293,7 @@ struct FApvMediaTmvEncoderOptions : public FTmvMediaEncoderOptions
 	//~ Begin FTmvMediaEncoderOptions
 	virtual FName GetEncoderName() const override { return FApvMediaTmvEncoderFactory::EncoderName;}
 	virtual FString GetFileSequenceExtension() const override { return TEXT("apv1"); }
+	virtual uint32 GetCodecFourCC() const override;
 	//~ End FTmvMediaEncoderOptions
 
 	/** Get currently configured color/chromat format. */
diff --git a/Engine/Plugins/Media/TmvMedia/Source/TmvMedia/Internal/Transcoder/TmvMediaContainerTranscodeMuxer.h b/Engine/Plugins/Media/TmvMedia/Source/TmvMedia/Internal/Transcoder/TmvMediaContainerTranscodeMuxer.h
new file mode 100644
--- /dev/null
+++ b/Engine/Plugins/Media/TmvMedia/Source/TmvMedia/Internal/Transcoder/TmvMediaContainerTranscodeMuxer.h
@@ -0,0 +1,127 @@
+// Copyright Epic Games, Inc. All Rights Reserved.
+
+#pragma once
+
+#include "Containers/Map.h"
+#include "HAL/Event.h"
+#include "Transcoder/TmvMediaTranscodeMuxer.h"
+#include "Encoder/ITmvMediaMuxer.h"
+
+#include "TmvMediaContainerTranscodeMuxer.generated.h"
+
+#define UE_API TMVMEDIA_API
+
+class ITmvMediaMuxer;
+
+/**
+ * Container muxer: muxes access units into a container file using ITmvMediaMuxer.
+ *
+ * Bridges the push-based transcoder pipeline (ReceiveAccessUnit) with the pull-based
+ * ITmvMediaMuxer (FSampleRequestDelegate). Incoming access units are buffered and
+ * delivered when the muxer's worker thread requests them.
+ *
+ * @see UTmvMediaTranscodeMuxer
+ */
+UCLASS(MinimalAPI)
+class UTmvMediaContainerTranscodeMuxer : public UTmvMediaTranscodeMuxer
+{
+	GENERATED_BODY()
+public:
+	UTmvMediaContainerTranscodeMuxer();
+
+	// ~Begin UTmvMediaTranscodeStage
+	UE_API virtual bool Start(UTmvMediaTranscodeJob* InParentJob) override;
+	UE_API virtual void RequestStop(UTmvMediaTranscodeJob* InParentJob) override;
+	// ~End UTmvMediaTranscodeStage
+
+	// ~Begin UTmvMediaTranscodeMuxer
+	UE_API virtual int32 OpenStream(UTmvMediaTranscodeJob* InParentJob, const FString& InStreamName, const FString& InExtension) override;
+	UE_API virtual void ReceiveAccessUnit(
+		UTmvMediaTranscodeJob* InParentJob,
+		int32 InStreamId,
+		const FTmvMediaFrameTimeInfo& TimeInfo,
+		TSharedPtr<TArray64<uint8>>&& InAccessUnit) override;
+	// ~End UTmvMediaTranscodeMuxer
+
+	/**
+	 * Set the track configuration for a stream opened with OpenStream.
+	 * Must be called before the first ReceiveAccessUnit for this stream.
+	 * @param InStreamId Stream index returned by OpenStream.
+	 * @param InTrackConfig Track configuration for the container.
+	 */
+	UE_API void SetStreamTrackConfig(int32 InStreamId, const FTmvMediaMuxerTrackConfig& InTrackConfig);
+
+private:
+	/** Buffered access unit waiting to be consumed by the muxer worker thread. */
+	struct FBufferedSample
+	{
+		TSharedPtr<TArray64<uint8>> Data;
+		uint32 SampleNumber = 0;
+	};
+
+	/** Try to fulfill pending muxer requests from the buffer. */
+	void FlushPendingRequests();
+
+	/** Extended stream info with container track mapping. */
+	struct FContainerStreamInfo
+	{
+		/** Track index in the ITmvMediaMuxer (-1 if not yet added). */
+		int32 MuxerTrackIndex = INDEX_NONE;
+
+		/** Track configuration (set via SetStreamTrackConfig). */
+		FTmvMediaMuxerTrackConfig TrackConfig;
+
+		/** Whether track config has been set. */
+		bool bHasTrackConfig = false;
+	};
+
+	TArray<FContainerStreamInfo> ContainerStreams;
+
+	/** The underlying container muxer implementation. */
+	TSharedPtr<ITmvMediaMuxer, ESPMode::ThreadSafe> Muxer;
+
+	/** Whether the muxer has been started (ITmvMediaMuxer::Start called). */
+	bool bMuxerStarted = false;
+	
+	/** Whether the muxer has failed to start and shouldn't be started again. */
+	bool bMuxerFailed = false;
+
+	/** Whether stop has been requested (signals final samples). */
+	std::atomic<bool> bStopRequested = false;
+
+	/**
+	 * Per-stream buffer state for the push-to-pull bridge.
+	 * One instance per stream, all fields protected by SampleBufferCS.
+	 */
+	struct FStreamBufferState
+	{
+		/** Buffered samples keyed by sample number. */
+		TMap<uint32, FBufferedSample> Samples;
+
+		/** Highest sample number received from the encoder (to detect last sample on stop). */
+		uint32 HighestReceivedSampleNumber = 0;
+
+		/** Set of sample numbers that have been successfully delivered to the muxer. */
+		TSet<uint32> DeliveredSampleNumbers;
+	};
+
+	/**
+	 * Critical section protecting StreamBuffers, MuxerTrackToStreamIndex,
+	 * PendingRequests, bMuxerStarted and bMuxerFailed.
+	 */
+	FCriticalSection SampleBufferCS;
+
+	/** Per-stream buffer state. One entry per stream opened with OpenStream. */
+	TArray<TUniquePtr<FStreamBufferState>> StreamBuffers;
+
+	/** Maps muxer track index to stream index. Populated in SetStreamTrackConfig. */
+	TMap<int32, int32> MuxerTrackToStreamIndex;
+
+	/** Pending requests from the muxer worker thread. May contain multiple requests per track. */
+	TArray<ITmvMediaMuxer::FSampleRequest> PendingRequests;
+
+	/** Event signaled by the muxer status delegate when muxing finishes or fails. */
+	FEventRef MuxerFinishedEvent;
+};
+
+#undef UE_API
diff --git a/Engine/Plugins/Media/TmvMedia/Source/TmvMedia/Internal/Transcoder/TmvMediaFileSequenceFrameProducer.h b/Engine/Plugins/Media/TmvMedia/Source/TmvMedia/Internal/Transcoder/TmvMediaFileSequenceFrameProducer.h
--- a/Engine/Plugins/Media/TmvMedia/Source/TmvMedia/Internal/Transcoder/TmvMediaFileSequenceFrameProducer.h
+++ b/Engine/Plugins/Media/TmvMedia/Source/TmvMedia/Internal/Transcoder/TmvMediaFileSequenceFrameProducer.h
@@ -26,7 +26,7 @@ private:
 	void ProcessAllImages(UTmvMediaTranscodeJob* InParentJob, const FString& InSequencePath);
 
 	/** Process one image. */
-	bool ProcessImage(UTmvMediaTranscodeJob* InParentJob, const FString& InImagePath);
+	bool ProcessImage(UTmvMediaTranscodeJob* InParentJob, const FString& InImagePath, int32 InFirstFrameIndex);
 
 	/** Returns true if the job is cancelled. */
 	bool IsCancelled(const UTmvMediaTranscodeJob* InParentJob) const;
diff --git a/Engine/Plugins/Media/TmvMedia/Source/TmvMedia/Internal/Transcoder/TmvMediaFileSequenceTranscodeMuxer.h b/Engine/Plugins/Media/TmvMedia/Source/TmvMedia/Internal/Transcoder/TmvMediaFileSequenceTranscodeMuxer.h
new file mode 100644
--- /dev/null
+++ b/Engine/Plugins/Media/TmvMedia/Source/TmvMedia/Internal/Transcoder/TmvMediaFileSequenceTranscodeMuxer.h
@@ -0,0 +1,29 @@
+// Copyright Epic Games, Inc. All Rights Reserved.
+
+#pragma once
+
+#include "Transcoder/TmvMediaTranscodeMuxer.h"
+
+#include "TmvMediaFileSequenceTranscodeMuxer.generated.h"
+
+#define UE_API TMVMEDIA_API
+
+/**
+ * File sequence muxer: writes each access unit as a separate file.
+ * @see UTmvMediaTranscodeMuxer
+ */
+UCLASS(MinimalAPI)
+class UTmvMediaFileSequenceTranscodeMuxer : public UTmvMediaTranscodeMuxer
+{
+	GENERATED_BODY()
+public:
+	UE_API virtual int32 OpenStream(UTmvMediaTranscodeJob* InParentJob, const FString& InStreamName, const FString& InExtension) override;
+
+	UE_API virtual void ReceiveAccessUnit(
+		UTmvMediaTranscodeJob* InParentJob,
+		int32 InStreamId,
+		const FTmvMediaFrameTimeInfo& TimeInfo,
+		TSharedPtr<TArray64<uint8>>&& InAccessUnit) override;
+};
+
+#undef UE_API
diff --git a/Engine/Plugins/Media/TmvMedia/Source/TmvMedia/Internal/Transcoder/TmvMediaTmvFrameEncoder.h b/Engine/Plugins/Media/TmvMedia/Source/TmvMedia/Internal/Transcoder/TmvMediaTmvFrameEncoder.h
--- a/Engine/Plugins/Media/TmvMedia/Source/TmvMedia/Internal/Transcoder/TmvMediaTmvFrameEncoder.h
+++ b/Engine/Plugins/Media/TmvMedia/Source/TmvMedia/Internal/Transcoder/TmvMediaTmvFrameEncoder.h
@@ -7,6 +7,7 @@
 
 class FTmvMediaTmvEncoderPool;
 class ITmvMediaEncoder;
+class UTmvMediaContainerTranscodeMuxer;
 
 /**
  * Implementation of an encoder stage of the transcode pipeline that wraps a TMV encoder.
@@ -41,4 +42,11 @@ private:
 
 	/** Name of the currently opened stream. */
 	FString StreamName;
-};
\ No newline at end of file
+
+	/** Pointer to container muxer if using container output (null for file sequence). */
+	UPROPERTY(Transient)
+	TObjectPtr<UTmvMediaContainerTranscodeMuxer> ContainerMuxer;
+
+	/** Whether the track config has been set on the container muxer. */
+	bool bContainerTrackConfigSet = false;
+};
diff --git a/Engine/Plugins/Media/TmvMedia/Source/TmvMedia/Private/TmvMediaModule.cpp b/Engine/Plugins/Media/TmvMedia/Source/TmvMedia/Private/TmvMediaModule.cpp
--- a/Engine/Plugins/Media/TmvMedia/Source/TmvMedia/Private/TmvMediaModule.cpp
+++ b/Engine/Plugins/Media/TmvMedia/Source/TmvMedia/Private/TmvMediaModule.cpp
@@ -1,7 +1,9 @@
 // Copyright Epic Games, Inc. All Rights Reserved.
 
 #include "Decoder/ITmvMediaDecoderFactory.h"
+#include "Decoder/ITmvMediaDemuxerFactory.h"
 #include "Encoder/ITmvMediaEncoderFactory.h"
+#include "Encoder/ITmvMediaMuxerFactory.h"
 #include "ITmvMediaModule.h"
 #include "Modules/ModuleManager.h"
 #include "TmvMediaLog.h"
@@ -11,6 +13,7 @@ DEFINE_LOG_CATEGORY(LogTmvMedia);
 
 class FTmvMediaModule : public ITmvMediaModule
 {
+public:
 	//~ Begin IModuleInterface
 	virtual void StartupModule() override
 	{
@@ -28,7 +31,7 @@ class FTmvMediaModule : public ITmvMediaModule
 	{
 		if (InDecoderFactory.IsValid())
 		{
-			FScopeLock lock(&Lock);
+			FScopeLock Lock(&FactoriesCS);
 			if (DecoderFactoryNames.Contains(InDecoderFactory->GetName()))
 			{
 				UE_LOGF(LogTmvMedia, Error, "Failed to add decoder factory \"%ls\", it is already registered.", *InDecoderFactory->GetName());
@@ -44,19 +47,19 @@ class FTmvMediaModule : public ITmvMediaModule
 	{
 		if (InDecoderFactory.IsValid())
 		{
-			FScopeLock lock(&Lock);
+			FScopeLock Lock(&FactoriesCS);
 			DecoderFactoryNames.Remove(InDecoderFactory->GetName());
 			DecoderFactories.RemoveSingleSwap(InDecoderFactory);
 		}
 	}
 	
-	virtual void GetDecoderFactories(TArray<TWeakPtr<ITmvMediaDecoderFactory, ESPMode::ThreadSafe>>& OutDecoderFactories) override
+	virtual void GetDecoderFactories(TArray<TWeakPtr<ITmvMediaDecoderFactory, ESPMode::ThreadSafe>>& OutDecoderFactories) const override
 	{
-		FScopeLock lock(&Lock);
+		FScopeLock Lock(&FactoriesCS);
 		OutDecoderFactories.Append(DecoderFactories);
 	}
 	
-	virtual TSharedPtr<ITmvMediaDecoderFactory, ESPMode::ThreadSafe> GetBestDecoderFactoryForFormat(const FString& InCodecFormat, const TMap<FString, FVariant>& InOptions) override
+	virtual TSharedPtr<ITmvMediaDecoderFactory, ESPMode::ThreadSafe> GetBestDecoderFactoryForFormat(const FString& InCodecFormat, const TMap<FString, FVariant>& InOptions) const override
 	{
 		TArray<TWeakPtr<ITmvMediaDecoderFactory, ESPMode::ThreadSafe>> Factories;
 		GetDecoderFactories(Factories);
@@ -81,7 +84,7 @@ class FTmvMediaModule : public ITmvMediaModule
 	{
 		if (InEncoderFactory.IsValid())
 		{
-			FScopeLock lock(&Lock);
+			FScopeLock Lock(&FactoriesCS);
 			if (EncoderFactoryNames.Contains(InEncoderFactory->GetName()))
 			{
 				UE_LOGF(LogTmvMedia, Error, "Failed to add encoder factory \"%ls\", it is already registered.", *InEncoderFactory->GetName().ToString());
@@ -96,7 +99,7 @@ class FTmvMediaModule : public ITmvMediaModule
 	{
 		if (InEncoderFactory.IsValid())
 		{
-			FScopeLock lock(&Lock);
+			FScopeLock Lock(&FactoriesCS);
 			EncoderFactoryNames.Remove(InEncoderFactory->GetName());
 			EncoderFactories.RemoveSingleSwap(InEncoderFactory);
 		}
@@ -104,13 +107,13 @@ class FTmvMediaModule : public ITmvMediaModule
 
 	virtual void GetEncoderFactories(TArray<TWeakPtr<ITmvMediaEncoderFactory, ESPMode::ThreadSafe>>& OutEncoderFactories) const override
 	{
-		FScopeLock lock(&Lock);
+		FScopeLock Lock(&FactoriesCS);
 		OutEncoderFactories.Append(EncoderFactories);
 	}
 
 	virtual TSharedPtr<ITmvMediaEncoderFactory, ESPMode::ThreadSafe> FindEncoderFactory(FName InEncoderFactoryName) const override
 	{
-		FScopeLock lock(&Lock);
+		FScopeLock Lock(&FactoriesCS);
 		for (TWeakPtr<ITmvMediaEncoderFactory, ESPMode::ThreadSafe> EncoderFactoryWeak : EncoderFactories)
 		{
 			if (TSharedPtr<ITmvMediaEncoderFactory, ESPMode::ThreadSafe> EncoderFactory = EncoderFactoryWeak.Pin())
@@ -124,6 +127,84 @@ class FTmvMediaModule : public ITmvMediaModule
 		return nullptr;
 	}
 
+	virtual void RegisterMuxerFactory(TSharedPtr<ITmvMediaMuxerFactory, ESPMode::ThreadSafe> InMuxerFactory) override
+	{
+		if (InMuxerFactory.IsValid())
+		{
+			FScopeLock Lock(&FactoriesCS);
+			if (MuxerFactoryNames.Contains(InMuxerFactory->GetName()))
+			{
+				UE_LOGF(LogTmvMedia, Error, "Failed to add muxer factory \"%ls\", it is already registered.", *InMuxerFactory->GetName().ToString());
+				return;
+			}
+			MuxerFactoryNames.Add(InMuxerFactory->GetName());
+			MuxerFactories.Emplace(MoveTemp(InMuxerFactory));
+		}
+	}
+
+	virtual void UnregisterMuxerFactory(TSharedPtr<ITmvMediaMuxerFactory, ESPMode::ThreadSafe> InMuxerFactory) override
+	{
+		if (InMuxerFactory.IsValid())
+		{
+			FScopeLock Lock(&FactoriesCS);
+			MuxerFactoryNames.Remove(InMuxerFactory->GetName());
+			MuxerFactories.RemoveSingleSwap(InMuxerFactory);
+		}
+	}
+
+	virtual void GetMuxerFactories(TArray<TWeakPtr<ITmvMediaMuxerFactory, ESPMode::ThreadSafe>>& OutMuxerFactories) const override
+	{
+		FScopeLock Lock(&FactoriesCS);
+		OutMuxerFactories.Append(MuxerFactories);
+	}
+
+	virtual void RegisterDemuxerFactory(TSharedPtr<ITmvMediaDemuxerFactory, ESPMode::ThreadSafe> InDemuxerFactory) override
+	{
+		if (InDemuxerFactory.IsValid())
+		{
+			FScopeLock Lock(&FactoriesCS);
+			if (DemuxerFactoryNames.Contains(InDemuxerFactory->GetName()))
+			{
+				UE_LOGF(LogTmvMedia, Error, "Failed to add demuxer factory \"%ls\", it is already registered.", *InDemuxerFactory->GetName().ToString());
+				return;
+			}
+			DemuxerFactoryNames.Add(InDemuxerFactory->GetName());
+			DemuxerFactories.Emplace(MoveTemp(InDemuxerFactory));
+		}
+	}
+
+	virtual void UnregisterDemuxerFactory(TSharedPtr<ITmvMediaDemuxerFactory, ESPMode::ThreadSafe> InDemuxerFactory) override
+	{
+		if (InDemuxerFactory.IsValid())
+		{
+			FScopeLock Lock(&FactoriesCS);
+			DemuxerFactoryNames.Remove(InDemuxerFactory->GetName());
+			DemuxerFactories.RemoveSingleSwap(InDemuxerFactory);
+		}
+	}
+
+	virtual void GetDemuxerFactories(TArray<TWeakPtr<ITmvMediaDemuxerFactory, ESPMode::ThreadSafe>>& OutDemuxerFactories) const override
+	{
+		FScopeLock Lock(&FactoriesCS);
+		OutDemuxerFactories.Append(DemuxerFactories);
+	}
+	
+	virtual TSharedPtr<ITmvMediaDemuxerFactory, ESPMode::ThreadSafe> FindDemuxerFactoryForExtension(const FString& InFileExtension) const override
+	{
+		FScopeLock Lock(&FactoriesCS);
+		for (const TWeakPtr<ITmvMediaDemuxerFactory, ESPMode::ThreadSafe>& FactoryWeak : DemuxerFactories)
+		{
+			if (TSharedPtr<ITmvMediaDemuxerFactory, ESPMode::ThreadSafe> Factory = FactoryWeak.Pin())
+			{
+				if (Factory->GetSupportedContainerFormats().Contains(InFileExtension))
+				{
+					return Factory;
+				}
+			}
+		}
+		return nullptr;
+	}
+
 	virtual ITmvMediaTranscodeJobManager* GetTranscodeJobManager() const override
 	{
 		return TranscodeJobManager.Get();
@@ -131,11 +212,15 @@ class FTmvMediaModule : public ITmvMediaModule
 	//~ End ITmvMediaModule
 
 private:
-	mutable FCriticalSection Lock;
+	mutable FCriticalSection FactoriesCS;
 	TArray<TWeakPtr<ITmvMediaDecoderFactory, ESPMode::ThreadSafe>> DecoderFactories;
 	TSet<FString> DecoderFactoryNames;
 	TArray<TWeakPtr<ITmvMediaEncoderFactory, ESPMode::ThreadSafe>> EncoderFactories;
 	TSet<FName> EncoderFactoryNames;
+	TArray<TWeakPtr<ITmvMediaMuxerFactory, ESPMode::ThreadSafe>> MuxerFactories;
+	TSet<FName> MuxerFactoryNames;
+	TArray<TWeakPtr<ITmvMediaDemuxerFactory, ESPMode::ThreadSafe>> DemuxerFactories;
+	TSet<FName> DemuxerFactoryNames;
 
 	TSharedPtr<FTmvMediaTranscodeJobManager> TranscodeJobManager;
 };
diff --git a/Engine/Plugins/Media/TmvMedia/Source/TmvMedia/Private/Transcoder/TmvMediaContainerTranscodeMuxer.cpp b/Engine/Plugins/Media/TmvMedia/Source/TmvMedia/Private/Transcoder/TmvMediaContainerTranscodeMuxer.cpp
new file mode 100644
--- /dev/null
+++ b/Engine/Plugins/Media/TmvMedia/Source/TmvMedia/Private/Transcoder/TmvMediaContainerTranscodeMuxer.cpp
@@ -0,0 +1,380 @@
+// Copyright Epic Games, Inc. All Rights Reserved.
+
+#include "Transcoder/TmvMediaContainerTranscodeMuxer.h"
+
+#include "Encoder/ITmvMediaMuxerFactory.h"
+#include "GenericPlatform/GenericPlatformFile.h"
+#include "HAL/PlatformFileManager.h"
+#include "ITmvMediaModule.h"
+#include "Misc/Paths.h"
+#include "TmvMediaLog.h"
+#include "Transcoder/TmvMediaFrameMips.h"
+#include "Transcoder/TmvMediaTranscodeJob.h"
+
+#include UE_INLINE_GENERATED_CPP_BY_NAME(TmvMediaContainerTranscodeMuxer)
+
+#define LOCTEXT_NAMESPACE "TmvMediaContainerTranscodeMuxer"
+
+UTmvMediaContainerTranscodeMuxer::UTmvMediaContainerTranscodeMuxer()
+	: MuxerFinishedEvent(EEventMode::ManualReset)
+{
+}
+
+bool UTmvMediaContainerTranscodeMuxer::Start(UTmvMediaTranscodeJob* InParentJob)
+{
+	if (!InParentJob)
+	{
+		UE_LOGF(LogTmvMedia, Error, "ContainerTranscodeMuxer::Start: Invalid parent job.");
+		return false;
+	}
+
+	bStopRequested = false;
+	MuxerFinishedEvent->Reset();
+
+	// Find the specified muxer factory.
+	const FName MuxerName(InParentJob->Settings.Muxer.Name);
+	const TSharedPtr<ITmvMediaMuxerFactory, ESPMode::ThreadSafe> MuxerFactory = FindMuxerFactoryByName(MuxerName);
+	if (!MuxerFactory.IsValid())
+	{
+		const FText Message = FText::Format(LOCTEXT("NoMuxerFactory", "No container muxer factory registered for \"{0}\"."), FText::FromName(MuxerName));
+		UE_LOGF(LogTmvMedia, Error, "ContainerTranscodeMuxer::Start: %ls", *Message.ToString());
+		UTmvMediaTranscodeJob::SafeReportError(InParentJob, Message);
+		return false;
+	}
+
+	Muxer = MuxerFactory->CreateMuxer();
+	if (!Muxer.IsValid())
+	{
+		const FText Message = LOCTEXT("CreateMuxerFailed", "Failed to create container muxer.");
+		UE_LOGF(LogTmvMedia, Error, "ContainerTranscodeMuxer::Start: %ls", *Message.ToString());
+		UTmvMediaTranscodeJob::SafeReportError(InParentJob, Message);
+		return false;
+	}
+
+	// Configure the muxer with output file path.
+	const FString OutputFilename = GetContainerOutputFilePath(InParentJob, MuxerFactory);
+	const FString OutPath = FPaths::GetPath(OutputFilename);
+
+	// Ensure output directory exists.
+	if (!FPaths::DirectoryExists(OutPath))
+	{
+		if (!FPlatformFileManager::Get().GetPlatformFile().CreateDirectoryTree(*OutPath))
+		{
+			const FText Message = FText::Format(LOCTEXT("FailedCreateOutputDir", "Failed to create output directory \"{0}\""), FText::FromString(OutPath));
+			UE_LOGF(LogTmvMedia, Error, "ContainerTranscodeMuxer::Start: %ls", *Message.ToString());
+			UTmvMediaTranscodeJob::SafeReportError(InParentJob, Message);
+			return false;
+		}
+	}
+
+	FTmvMediaMuxerConfig Config;
+	Config.OutputFilename = OutputFilename;
+	Config.OutputMode = FTmvMediaMuxerConfig::EOutputMode::Standard;
+
+	if (Muxer->Configure(Config) != ETmvMediaContainerResult::Success)
+	{
+		const FText Message = FText::Format(LOCTEXT("ConfigureFailed", "Failed to configure container muxer: {0}"), FText::FromString(Muxer->GetLastError()));
+		UE_LOGF(LogTmvMedia, Error, "ContainerTranscodeMuxer::Start: %ls", *Message.ToString());
+		UTmvMediaTranscodeJob::SafeReportError(InParentJob, Message);
+		return false;
+	}
+
+	return Super::Start(InParentJob);
+}
+
+void UTmvMediaContainerTranscodeMuxer::RequestStop(UTmvMediaTranscodeJob* InParentJob)
+{
+	bStopRequested = true;
+
+	if (Muxer.IsValid() && bMuxerStarted)
+	{
+		// Flush any remaining buffered samples now that bStopRequested is true,
+		// which will mark the last deliverable sample with bIsFinalSample.
+		// For any pending requests that can't be fulfilled from the buffer
+		// (worker already consumed the last real sample and is asking for the next),
+		// send an empty final sample to unblock the worker.
+		{
+			FScopeLock Lock(&SampleBufferCS);
+			FlushPendingRequests();
+
+			for (const ITmvMediaMuxer::FSampleRequest& Request : PendingRequests)
+			{
+				FTmvMediaMuxerSample FinalSample;
+				FinalSample.SampleNumber = Request.SampleNumber;
+				FinalSample.bIsFinalSample = true;
+				if (Muxer->AddSample(Request.TrackIndex, FinalSample) != ETmvMediaContainerResult::Success)
+				{
+					UE_LOGF(LogTmvMedia, Error, "ContainerTranscodeMuxer::RequestStop: Failed to add final sample for track %d: %ls", Request.TrackIndex, *Muxer->GetLastError());
+				}
+			}
+			PendingRequests.Empty();
+		}
+	}
+
+	if (Muxer.IsValid() && bMuxerStarted)
+	{
+		// Wait for the muxer worker to finish writing all samples.
+		// The status delegate signals MuxerFinishedEvent on Finished or Failed.
+		constexpr uint32 TimeoutMs = 30000;
+		bool bTimedOut = false;
+		if (!MuxerFinishedEvent->Wait(TimeoutMs))
+		{
+			UE_LOGF(LogTmvMedia, Warning, "ContainerTranscodeMuxer::RequestStop: Timed out waiting for muxer to finish. Aborting.");
+			bTimedOut = true;
+		}
+
+		if (Muxer->Finalize() != ETmvMediaContainerResult::Success)
+		{
+			UE_LOGF(LogTmvMedia, Error, "ContainerTranscodeMuxer::RequestStop: Finalize failed: %ls", *Muxer->GetLastError());
+			UTmvMediaTranscodeJob::SafeReportError(InParentJob, FText::FromString(Muxer->GetLastError()));
+		}
+		else if (bTimedOut)
+		{
+			UTmvMediaTranscodeJob::SafeReportError(InParentJob, LOCTEXT("ErrorMuxerTimeout", "Timed out waiting for muxer to finish. Aborting."));
+		}
+
+		bMuxerStarted = false;
+		bMuxerFailed = false;
+	}
+
+	if (Muxer.IsValid())
+	{
+		Muxer.Reset();
+	}
+
+	ContainerStreams.Empty();
+	StreamBuffers.Empty();
+	MuxerTrackToStreamIndex.Empty();
+	PendingRequests.Empty();
+
+	Super::RequestStop(InParentJob);
+}
+
+int32 UTmvMediaContainerTranscodeMuxer::OpenStream(UTmvMediaTranscodeJob* InParentJob, const FString& InStreamName, const FString& InExtension)
+{
+	if (!InParentJob)
+	{
+		return INDEX_NONE;
+	}
+
+	FScopeLock Lock(&StreamsCS);
+
+	Streams.Add({InStreamName, InExtension});
+	ContainerStreams.AddDefaulted();
+	StreamBuffers.Add(MakeUnique<FStreamBufferState>());
+
+	return Streams.Num() - 1;
+}
+
+void UTmvMediaContainerTranscodeMuxer::SetStreamTrackConfig(int32 InStreamId, const FTmvMediaMuxerTrackConfig& InTrackConfig)
+{
+	int32 MuxerTrackIndex = INDEX_NONE;
+	
+	{
+		FScopeLock Lock(&StreamsCS);
+
+		if (!ContainerStreams.IsValidIndex(InStreamId))
+		{
+			UE_LOGF(LogTmvMedia, Error, "ContainerTranscodeMuxer::SetStreamTrackConfig: Invalid stream id %d.", InStreamId);
+			return;
+		}
+
+		if (!Muxer.IsValid())
+		{
+			UE_LOGF(LogTmvMedia, Error, "ContainerTranscodeMuxer::SetStreamTrackConfig: Muxer not initialized.");
+			return;
+		}
+
+		FContainerStreamInfo& Info = ContainerStreams[InStreamId];
+		Info.TrackConfig = InTrackConfig;
+		Info.bHasTrackConfig = true;
+
+		// Add track to the underlying muxer.
+		MuxerTrackIndex = Info.MuxerTrackIndex = Muxer->AddTrack(InTrackConfig);
+		if (MuxerTrackIndex < 0)
+		{
+			UE_LOGF(LogTmvMedia, Error, "ContainerTranscodeMuxer::SetStreamTrackConfig: Failed to add track: %ls", *Muxer->GetLastError());
+		}
+	}
+
+	if (MuxerTrackIndex >= 0)
+	{
+		// Register the mapping so FlushPendingRequests can resolve track index
+		// to stream index without acquiring StreamsCS.
+		FScopeLock Lock(&SampleBufferCS);
+		MuxerTrackToStreamIndex.Add(MuxerTrackIndex, InStreamId);
+	}
+}
+
+void UTmvMediaContainerTranscodeMuxer::ReceiveAccessUnit(
+		UTmvMediaTranscodeJob* InParentJob,
+		int32 InStreamId,
+		const FTmvMediaFrameTimeInfo& InTimeInfo,
+		TSharedPtr<TArray64<uint8>>&& InAccessUnit)
+{
+	if (!InAccessUnit.IsValid() || InAccessUnit->Num() == 0)
+	{
+		return;
+	}
+
+	if (!Muxer.IsValid())
+	{
+		UE_LOGF(LogTmvMedia, Error, "ContainerTranscodeMuxer::ReceiveAccessUnit: Muxer not initialized.");
+		return;
+	}
+
+	{
+		FScopeLock Lock(&StreamsCS);
+		if (!ContainerStreams.IsValidIndex(InStreamId))
+		{
+			const FText Message = FText::Format(LOCTEXT("InvalidStreamId", "Invalid stream id ({0})"), FText::AsNumber(InStreamId));
+			UE_LOGF(LogTmvMedia, Error, "ContainerTranscodeMuxer::ReceiveAccessUnit: %ls.", *Message.ToString());
+			UTmvMediaTranscodeJob::SafeReportError(InParentJob, Message);
+			return;
+		}
+
+		const FContainerStreamInfo& Info = ContainerStreams[InStreamId];
+		if (!Info.bHasTrackConfig || Info.MuxerTrackIndex < 0)
+		{
+			const FText Message = FText::Format(LOCTEXT("InvalidTrackConfig", "Track config not set for stream {0}"), FText::AsNumber(InStreamId));
+			UE_LOGF(LogTmvMedia, Error, "ContainerTranscodeMuxer::ReceiveAccessUnit: %ls.", *Message.ToString());
+			UTmvMediaTranscodeJob::SafeReportError(InParentJob, Message);
+			return;
+		}
+	}
+
+	// Buffer the incoming access unit and lazily start the muxer on the first sample.
+	{
+		FScopeLock Lock(&SampleBufferCS);
+
+		// Start the muxer on first sample if not yet started.
+		// Protected by SampleBufferCS to prevent multiple threads from calling Start() concurrently.
+		if (!bMuxerStarted && !bMuxerFailed)
+		{
+			ITmvMediaMuxer::FSampleRequestDelegate SampleRequestDelegate;
+			SampleRequestDelegate.BindWeakLambda(this, [this](const TArray<ITmvMediaMuxer::FSampleRequest>& InRequests)
+			{
+				// Called from the muxer's worker thread when it needs samples.
+				FScopeLock InnerLock(&SampleBufferCS);
+				PendingRequests.Append(InRequests);
+				// Try to fulfill from already-buffered data.
+				FlushPendingRequests();
+			});
+
+			ITmvMediaMuxer::FStatusDelegate StatusDelegate;
+			TWeakObjectPtr<UTmvMediaTranscodeJob> ParentJobWeak = InParentJob;
+			StatusDelegate.BindWeakLambda(this, [this, ParentJobWeak](ITmvMediaMuxer::EStatus InStatus)
+			{
+				if (InStatus == ITmvMediaMuxer::EStatus::Failed && Muxer.IsValid())
+				{
+					const FText Message = FText::Format(LOCTEXT("MuxerFailure", "Muxer reported failure: {0}"), FText::FromString(Muxer->GetLastError()));
+					UE_LOGF(LogTmvMedia, Error, "ContainerTranscodeMuxer: %ls.", *Message.ToString());
+					UTmvMediaTranscodeJob::SafeReportError(ParentJobWeak.Get(), Message);
+				}
+				// Signal that the muxer worker has finished (success or failure).
+				// RequestStop() waits on this before calling Finalize().
+				MuxerFinishedEvent->Trigger();
+			});
+
+			if (Muxer->Start(SampleRequestDelegate, StatusDelegate) != ETmvMediaContainerResult::Success)
+			{
+				const FText Message = FText::Format(LOCTEXT("MuxerStartFailure", "Failed to start muxer: {0}"), FText::FromString(Muxer->GetLastError()));
+				UE_LOGF(LogTmvMedia, Error, "ContainerTranscodeMuxer::ReceiveAccessUnit: %ls.", *Message.ToString());
+				UTmvMediaTranscodeJob::SafeReportError(InParentJob, Message);
+				bMuxerFailed = true;
+				return;
+			}
+			bMuxerStarted = true;
+		}
+
+		if (StreamBuffers.IsValidIndex(InStreamId) && StreamBuffers[InStreamId].IsValid())
+		{
+			// For the sample number, we need contiguous values that start at index 0.
+			const uint32 SampleNumber = static_cast<uint32>(InTimeInfo.FrameIndexNoOffset);
+
+			FStreamBufferState& State = *StreamBuffers[InStreamId];
+
+			FBufferedSample Buffered;
+			Buffered.Data = MoveTemp(InAccessUnit);
+			Buffered.SampleNumber = SampleNumber;
+			State.Samples.Add(SampleNumber, MoveTemp(Buffered));
+
+			if (SampleNumber >= State.HighestReceivedSampleNumber)
+			{
+				State.HighestReceivedSampleNumber = SampleNumber + 1;
+			}
+		}
+
+		// Try to fulfill any pending requests now that we have new data.
+		FlushPendingRequests();
+	}
+}
+
+void UTmvMediaContainerTranscodeMuxer::FlushPendingRequests()
+{
+	// Must be called with SampleBufferCS held.
+	// Collect indices of fulfilled/stale requests, then remove after iteration.
+	TArray<int32, TInlineAllocator<16>> IndicesToRemove;
+
+	for (int32 i = 0; i < PendingRequests.Num(); ++i)
+	{
+		const ITmvMediaMuxer::FSampleRequest& Request = PendingRequests[i];
+
+		// Resolve muxer track index to stream index via pre-built map (no StreamsCS needed).
+		const int32* StreamIndexPtr = MuxerTrackToStreamIndex.Find(Request.TrackIndex);
+		if (!StreamIndexPtr || !StreamBuffers.IsValidIndex(*StreamIndexPtr) || !StreamBuffers[*StreamIndexPtr].IsValid())
+		{
+			UE_LOGF(LogTmvMedia, Error, 
+				"ContainerTranscodeMuxer::FlushPendingRequests: Failed to find StreamBuffer for track index %d."
+				" Pending Request for frame %d will be ignored.", Request.TrackIndex, Request.SampleNumber);
+			IndicesToRemove.Add(i);
+			continue;
+		}
+
+		FStreamBufferState& State = *StreamBuffers[*StreamIndexPtr];
+
+		// Discard duplicate requests for samples that have already been delivered.
+		if (State.DeliveredSampleNumbers.Contains(Request.SampleNumber))
+		{
+			IndicesToRemove.Add(i);
+			continue;
+		}
+
+		if (const FBufferedSample* FoundSample = State.Samples.Find(Request.SampleNumber))
+		{
+			FTmvMediaMuxerSample Sample;
+			Sample.Data = MakeArrayView(FoundSample->Data->GetData(), FoundSample->Data->Num());
+			Sample.SampleNumber = FoundSample->SampleNumber;
+			Sample.bIsKeyframe = true;
+
+			// Mark as final if stop was requested and this is the last received sample.
+			if (bStopRequested && (Request.SampleNumber + 1) >= State.HighestReceivedSampleNumber)
+			{
+				UE_LOGF(LogTmvMedia, Verbose, "ContainerTranscodeMuxer::FlushPendingRequests: Last sample %u for track %d", FoundSample->SampleNumber, Request.TrackIndex);
+				Sample.bIsFinalSample = true;
+			}
+
+			if (Muxer->AddSample(Request.TrackIndex, Sample) != ETmvMediaContainerResult::Success)
+			{
+				UE_LOGF(LogTmvMedia, Error, "ContainerTranscodeMuxer::FlushPendingRequests: Failed to add sample %u for track %d: %ls",
+					FoundSample->SampleNumber, Request.TrackIndex, *Muxer->GetLastError());
+				// Keep the sample in the buffer and the request pending so it can be retried.
+				continue;
+			}
+			
+			UE_LOGF(LogTmvMedia, Verbose, "ContainerTranscodeMuxer::FlushPendingRequests: Added sample %u for track %d", FoundSample->SampleNumber, Request.TrackIndex);
+
+			State.Samples.Remove(Request.SampleNumber);
+			State.DeliveredSampleNumbers.Add(Request.SampleNumber);
+			IndicesToRemove.Add(i);
+		}
+	}
+
+	// Remove in reverse order so indices stay valid.
+	for (int32 i = IndicesToRemove.Num() - 1; i >= 0; --i)
+	{
+		PendingRequests.RemoveAtSwap(IndicesToRemove[i]);
+	}
+}
+
+#undef LOCTEXT_NAMESPACE
diff --git a/Engine/Plugins/Media/TmvMedia/Source/TmvMedia/Private/Transcoder/TmvMediaFileSequenceFrameProducer.cpp b/Engine/Plugins/Media/TmvMedia/Source/TmvMedia/Private/Transcoder/TmvMediaFileSequenceFrameProducer.cpp
--- a/Engine/Plugins/Media/TmvMedia/Source/TmvMedia/Private/Transcoder/TmvMediaFileSequenceFrameProducer.cpp
+++ b/Engine/Plugins/Media/TmvMedia/Source/TmvMedia/Private/Transcoder/TmvMediaFileSequenceFrameProducer.cpp
@@ -182,6 +182,12 @@ void UTmvMediaFileSequenceFrameProducer::ProcessAllImages(UTmvMediaTranscodeJob*
 	// Upper bound for this to avoid overloading the system. Todo: Make this configurable?
 	NumThreads = FMath::Min(NumThreads, FPlatformMisc::NumberOfCores());
 
+	int32 FirstFrameIndex = 0;
+	if (FoundFiles.Num() > 0)
+	{
+		UE::TmvMedia::GetFrameNumberFromFilename(FoundFiles[0], FirstFrameIndex);
+	}
+	
 	// Loop through all files.
 	int NumDone = 0;
 	int TotalNum = FoundFiles.Num();
@@ -211,13 +217,13 @@ void UTmvMediaFileSequenceFrameProducer::ProcessAllImages(UTmvMediaTranscodeJob*
 		TWeakObjectPtr<UTmvMediaTranscodeJob> ParentJobWeak = InParentJob;
 		TWeakObjectPtr<UTmvMediaFileSequenceFrameProducer> SelfWeak(this);
 
-		Async(ExecutionTarget, [SelfWeak, FullFileName, ParentJobWeak]() mutable
+		Async(ExecutionTarget, [SelfWeak, FullFileName, ParentJobWeak, FirstFrameIndex]() mutable
 		{
 			if (TStrongObjectPtr<UTmvMediaFileSequenceFrameProducer> Self = SelfWeak.Pin())
 			{
 				if (TStrongObjectPtr<UTmvMediaTranscodeJob> ParentJob = ParentJobWeak.Pin())
 				{
-					Self->ProcessImage(ParentJob.Get(), FullFileName);
+					Self->ProcessImage(ParentJob.Get(), FullFileName, FirstFrameIndex);
 				}
 				--Self->NumActiveFrames;
 			}
@@ -242,7 +248,7 @@ bool UTmvMediaFileSequenceFrameProducer::IsCancelled(const UTmvMediaTranscodeJob
 	return !IsValid(InParentJob) || !InParentJob->IsRunning() || GetStageStatus() != ETmvMediaTranscodeStageStatus::Started;
 }
 
-bool UTmvMediaFileSequenceFrameProducer::ProcessImage(UTmvMediaTranscodeJob* InParentJob, const FString& InImagePath)
+bool UTmvMediaFileSequenceFrameProducer::ProcessImage(UTmvMediaTranscodeJob* InParentJob, const FString& InImagePath, int32 InFirstFrameIndex)
 {
 	TSharedPtr<FImage> Image = MakeShared<FImage>();
 
@@ -268,6 +274,7 @@ bool UTmvMediaFileSequenceFrameProducer::ProcessImage(UTmvMediaTranscodeJob* InP
 		TUniquePtr<FTmvMediaFrameMips> FrameMips = MakeUnique<FTmvMediaFrameMips>();
 		// Extract frame number from the name.
 		UE::TmvMedia::GetFrameNumberFromFilename(InImagePath, FrameMips->TimeInfo.FrameIndex);
+		FrameMips->TimeInfo.FrameIndexNoOffset = FrameMips->TimeInfo.FrameIndex - InFirstFrameIndex;
 		FrameMips->TimeInfo.FrameDuration = FTimespan::FromSeconds(InParentJob->Settings.FrameRate.AsInterval());
 		FrameMips->TimeInfo.FrameTime = FrameMips->TimeInfo.FrameDuration * FrameMips->TimeInfo.FrameIndex;
 		FrameMips->Init(Image, InParentJob->Settings.bEnableMipMapping);
diff --git a/Engine/Plugins/Media/TmvMedia/Source/TmvMedia/Private/Transcoder/TmvMediaFileSequenceTranscodeMuxer.cpp b/Engine/Plugins/Media/TmvMedia/Source/TmvMedia/Private/Transcoder/TmvMediaFileSequenceTranscodeMuxer.cpp
new file mode 100644
--- /dev/null
+++ b/Engine/Plugins/Media/TmvMedia/Source/TmvMedia/Private/Transcoder/TmvMediaFileSequenceTranscodeMuxer.cpp
@@ -0,0 +1,120 @@
+// Copyright Epic Games, Inc. All Rights Reserved.
+
+#include "Transcoder/TmvMediaFileSequenceTranscodeMuxer.h"
+
+#include "GenericPlatform/GenericPlatformFile.h"
+#include "HAL/PlatformFileManager.h"
+#include "Misc/Paths.h"
+#include "TmvMediaLog.h"
+#include "Transcoder/TmvMediaFrameMips.h"
+#include "Transcoder/TmvMediaTranscodeJob.h"
+#include "Utils/TmvMediaPathUtils.h"
+
+#include UE_INLINE_GENERATED_CPP_BY_NAME(TmvMediaFileSequenceTranscodeMuxer)
+
+#define LOCTEXT_NAMESPACE "TmvMediaFileSequenceTranscodeMuxer"
+
+int32 UTmvMediaFileSequenceTranscodeMuxer::OpenStream(UTmvMediaTranscodeJob* InParentJob, const FString& InStreamName, const FString& InExtension)
+{
+	if (!InParentJob)
+	{
+		return INDEX_NONE;
+	}
+
+	if (InStreamName.IsEmpty() || InExtension.IsEmpty())
+	{
+		const FText Message = FText::Format(LOCTEXT("InvalidStream", "Invalid stream name or extension (Name=\"{0}\", Ext=\"{1}\""),
+			FText::FromString(InStreamName), FText::FromString(InExtension));
+		UE_LOGF(LogTmvMedia, Error, "FileSequenceTranscodeMuxer::OpenStream: %ls.", *Message.ToString());
+		UTmvMediaTranscodeJob::SafeReportError(InParentJob, Message);
+		return INDEX_NONE;
+	}
+
+	const FString OutPath = InParentJob->Settings.GetAbsoluteOutputPath();
+
+	if (OutPath.IsEmpty())
+	{
+		const FText Message = LOCTEXT("EmptyOutputPath", "OutputPath is empty");
+		UE_LOGF(LogTmvMedia, Error, "FileSequenceTranscodeMuxer::OpenStream: %ls.", *Message.ToString());
+		UTmvMediaTranscodeJob::SafeReportError(InParentJob, Message);
+		return INDEX_NONE;
+	}
+
+	if (!FPaths::DirectoryExists(OutPath))
+	{
+		if (!FPlatformFileManager::Get().GetPlatformFile().CreateDirectoryTree(*OutPath))
+		{
+			const FText Message = FText::Format(LOCTEXT("FailedCreateOutputDir", "Failed to create output directory \"{0}\""), FText::FromString(OutPath));
+			UE_LOGF(LogTmvMedia, Error, "FileSequenceTranscodeMuxer::OpenStream: %ls.", *Message.ToString());
+			UTmvMediaTranscodeJob::SafeReportError(InParentJob, Message);
+			return INDEX_NONE;
+		}
+	}
+
+	FScopeLock Lock(&StreamsCS);
+	Streams.Add({InStreamName, InExtension});
+	return Streams.Num() - 1;
+}
+
+void UTmvMediaFileSequenceTranscodeMuxer::ReceiveAccessUnit(
+		UTmvMediaTranscodeJob* InParentJob,
+		int32 InStreamId,
+		const FTmvMediaFrameTimeInfo& InTimeInfo,
+		TSharedPtr<TArray64<uint8>>&& InAccessUnit)
+{
+	if (!InAccessUnit.IsValid())
+	{
+		const FText Message = LOCTEXT("InvalidAccessUnit", "Null access unit");
+		UE_LOGF(LogTmvMedia, Error, "FileSequenceTranscodeMuxer::ReceiveAccessUnit: %ls.", *Message.ToString());
+		UTmvMediaTranscodeJob::SafeReportError(InParentJob, Message);
+		return;
+	}
+
+	if (InAccessUnit->Num() == 0)
+	{
+		UE_LOGF(LogTmvMedia, Warning, "ReceiveAccessUnit: Empty access unit (stream %d, frame %d).", InStreamId, InTimeInfo.FrameIndex);
+		return;
+	}
+
+	FStreamInfo StreamInfo;
+	{
+		FScopeLock Lock(&StreamsCS);
+		if (!Streams.IsValidIndex(InStreamId))
+		{
+			const FText Message = LOCTEXT("InvalidStreamId", "Invalid stream id");
+			UE_LOGF(LogTmvMedia, Error, "FileSequenceTranscodeMuxer::ReceiveAccessUnit: %ls.", *Message.ToString());
+			UTmvMediaTranscodeJob::SafeReportError(InParentJob, Message);
+			return;
+		}
+		StreamInfo = Streams[InStreamId];
+	}
+
+	if (!InParentJob)
+	{
+		UE_LOGF(LogTmvMedia, Error, "ReceiveAccessUnit: Invalid Transcode Job.");
+		return;
+	}
+
+	const FString OutPath = InParentJob->Settings.GetAbsoluteOutputPath();
+	const FString FileName = UE::TmvMedia::PathUtils::MakeFrameFilename(*StreamInfo.StreamName, InTimeInfo.FrameIndex, InParentJob->Settings.ZeroPadFrameNumbers, *StreamInfo.Extension);
+	const FString FullFileName = FPaths::Combine(OutPath, FileName);
+
+	TUniquePtr<IFileHandle> File(FPlatformFileManager::Get().GetPlatformFile().OpenWrite(*FullFileName));
+	if (File)
+	{
+		if (!File->Write(InAccessUnit->GetData(), InAccessUnit->Num()))
+		{
+			const FText Message = FText::Format(LOCTEXT("WriteFailed", "Failed to write access unit to file \"{0}\""), FText::FromString(FullFileName));
+			UE_LOGF(LogTmvMedia, Error, "FileSequenceTranscodeMuxer::ReceiveAccessUnit: %ls.", *Message.ToString());
+			UTmvMediaTranscodeJob::SafeReportError(InParentJob, Message);
+		}
+	}
+	else
+	{
+		const FText Message = FText::Format(LOCTEXT("OpenFailed", "Failed to open output file for writing: \"{0}\""), FText::FromString(FullFileName));
+		UE_LOGF(LogTmvMedia, Error, "FileSequenceTranscodeMuxer::ReceiveAccessUnit: %ls.", *Message.ToString());
+		UTmvMediaTranscodeJob::SafeReportError(InParentJob, Message);
+	}
+}
+
+#undef LOCTEXT_NAMESPACE
diff --git a/Engine/Plugins/Media/TmvMedia/Source/TmvMedia/Private/Transcoder/TmvMediaPlayerFrameProducerResource.cpp b/Engine/Plugins/Media/TmvMedia/Source/TmvMedia/Private/Transcoder/TmvMediaPlayerFrameProducerResource.cpp
--- a/Engine/Plugins/Media/TmvMedia/Source/TmvMedia/Private/Transcoder/TmvMediaPlayerFrameProducerResource.cpp
+++ b/Engine/Plugins/Media/TmvMedia/Source/TmvMedia/Private/Transcoder/TmvMediaPlayerFrameProducerResource.cpp
@@ -348,7 +348,7 @@ void FTmvMediaPlayerFrameProducerResource::OnMediaTextureRendered_RenderThread(c
 		UE_LOGF(LogTmvMedia, Error,
 			"Transcode Frame Producer: Sample Time %ls earlier than last rendered time %ls.",
 			*InMediaTextureSample->GetTime().Time.ToString(),
-			*LastRenderedSampleTime.Load().ToString());
+			*LastRenderedSampleTime.load().ToString());
 		return;
 	}
 
@@ -367,6 +367,7 @@ void FTmvMediaPlayerFrameProducerResource::OnMediaTextureRendered_RenderThread(c
 	{
 		const FTimespan RelativeFrameTime = TimeInfo.FrameTime - FirstRenderedSampleTime;
 		TimeInfo.FrameIndex = static_cast<int32>(FMath::RoundHalfToEven(RelativeFrameTime.GetTotalSeconds() / TimeInfo.FrameDuration.GetTotalSeconds()));
+		TimeInfo.FrameIndexNoOffset = TimeInfo.FrameIndex;	// This should already be zero based. 
 	}
 
 	FRHICommandListImmediate& RhiCmdList = FRHICommandListImmediate::Get();
diff --git a/Engine/Plugins/Media/TmvMedia/Source/TmvMedia/Private/Transcoder/TmvMediaPlayerFrameProducerResource.h b/Engine/Plugins/Media/TmvMedia/Source/TmvMedia/Private/Transcoder/TmvMediaPlayerFrameProducerResource.h
--- a/Engine/Plugins/Media/TmvMedia/Source/TmvMedia/Private/Transcoder/TmvMediaPlayerFrameProducerResource.h
+++ b/Engine/Plugins/Media/TmvMedia/Source/TmvMedia/Private/Transcoder/TmvMediaPlayerFrameProducerResource.h
@@ -106,10 +106,10 @@ private:
 	FTimespan FirstRenderedSampleTime;
 
 	/** This is updated on the render thread to keep track of the rendered "play head". */
-	TAtomic<FTimespan> LastRenderedSampleTime;
+	std::atomic<FTimespan> LastRenderedSampleTime;
 
 	/** This is updated on the render thread as additional information for controller the player. */
-	TAtomic<FTimespan> LastRenderedSampleDuration;
+	std::atomic<FTimespan> LastRenderedSampleDuration;
 
 	/** Number of frames submitted to the converter/decoder (in flight after the producer). */
 	std::atomic_int32_t SubmittedCount = 0;
diff --git a/Engine/Plugins/Media/TmvMedia/Source/TmvMedia/Private/Transcoder/TmvMediaTmvFrameEncoder.cpp b/Engine/Plugins/Media/TmvMedia/Source/TmvMedia/Private/Transcoder/TmvMediaTmvFrameEncoder.cpp
--- a/Engine/Plugins/Media/TmvMedia/Source/TmvMedia/Private/Transcoder/TmvMediaTmvFrameEncoder.cpp
+++ b/Engine/Plugins/Media/TmvMedia/Source/TmvMedia/Private/Transcoder/TmvMediaTmvFrameEncoder.cpp
@@ -15,6 +15,8 @@
 #include "TmvMediaLog.h"
 #include "Transcoder/TmvMediaFrameMips.h"
 #include "Transcoder/TmvMediaTranscodeJob.h"
+#include "Transcoder/TmvMediaContainerTranscodeMuxer.h"
+#include "Transcoder/TmvMediaFrameProducer.h"
 #include "Transcoder/TmvMediaTranscodeMuxer.h"
 #include "Utils/TmvMediaMessageContext.h"
 #include "Utils/TTmvMediaSharedObjectPool.h"
@@ -173,6 +175,10 @@ bool UTmvMediaTmvFrameEncoder::Start(UTmvMediaTranscodeJob* InParentJob)
 		if (UTmvMediaTranscodeMuxer* Muxer = InParentJob->GetStage<UTmvMediaTranscodeMuxer>())
 		{
 			StreamId = Muxer->OpenStream(InParentJob, StreamName, EncoderOptions.Get().GetFileSequenceExtension());
+
+			// Check if this is a container muxer for deferred track config.
+			ContainerMuxer = Cast<UTmvMediaContainerTranscodeMuxer>(Muxer);
+			bContainerTrackConfigSet = false;
 		}
 
 		if (StreamId == INDEX_NONE)
@@ -274,6 +280,32 @@ bool UTmvMediaTmvFrameEncoder::RequestMipInfos(UTmvMediaTranscodeJob* InParentJo
 		UTmvMediaTranscodeJob::SafeReportError(InParentJob, Message);
 		return false;
 	}
+
+	// Set the container muxer track config on first call now that we know dimensions.
+	if (ContainerMuxer && !bContainerTrackConfigSet && StreamId != INDEX_NONE)
+	{
+		FTmvMediaMuxerTrackConfig TrackConfig;
+		TrackConfig.TrackType = ETmvMediaTrackType::Video;
+		TrackConfig.SampleEntryFormat = EncoderOptions.Get().GetCodecFourCC();
+		TrackConfig.DisplayWidth = static_cast<uint32>(InSourceMipInfo.Width);
+		TrackConfig.DisplayHeight = static_cast<uint32>(InSourceMipInfo.Height);
+		TrackConfig.bIsAllKeyframes = true;
+
+		// Set timescale and constant duration from frame rate.
+		FFrameRate FrameRate = InParentJob->Settings.GetInputFramerate();
+		// Fetch the frame rate from the frame producer, which is going to be what the media file was.
+		if (UTmvMediaFrameProducer* FrameProducer = InParentJob->GetStage<UTmvMediaFrameProducer>())
+		{
+			const FTmvMediaFrameProducerTrackInfo VideoTrackInfo = FrameProducer->GetVideoTrackInfo();
+			FrameRate = VideoTrackInfo.FrameRate;
+		}
+		TrackConfig.Timescale = static_cast<uint32>(FrameRate.Numerator);
+		TrackConfig.ConstantSampleDuration = static_cast<uint32>(FrameRate.Denominator);
+
+		ContainerMuxer->SetStreamTrackConfig(StreamId, TrackConfig);
+		bContainerTrackConfigSet = true;
+	}
+
 	return true;
 }
 
diff --git a/Engine/Plugins/Media/TmvMedia/Source/TmvMedia/Private/Transcoder/TmvMediaTranscodeList.cpp b/Engine/Plugins/Media/TmvMedia/Source/TmvMedia/Private/Transcoder/TmvMediaTranscodeList.cpp
--- a/Engine/Plugins/Media/TmvMedia/Source/TmvMedia/Private/Transcoder/TmvMediaTranscodeList.cpp
+++ b/Engine/Plugins/Media/TmvMedia/Source/TmvMedia/Private/Transcoder/TmvMediaTranscodeList.cpp
@@ -2,8 +2,10 @@
 
 #include "Transcoder/TmvMediaTranscodeList.h"
 
+#include "ITmvMediaModule.h"
 #include "Algo/Sort.h"
 #include "Algo/Unique.h"
+#include "Encoder/ITmvMediaMuxerFactory.h"
 #include "UObject/UObjectIterator.h"
 
 #include UE_INLINE_GENERATED_CPP_BY_NAME(TmvMediaTranscodeList)
@@ -37,6 +39,27 @@ namespace UE::TmvMedia::Transcode
 
 		return FoundStruct;
 	}
+	
+	/** Returns the name of the first registered muxer factory. */ 
+	FName GetDefaultMuxerName()
+	{
+		const ITmvMediaModule* TmvMediaModule = ITmvMediaModule::Get();
+		if (!TmvMediaModule)
+		{
+			return NAME_None;
+		}
+
+		TArray<TWeakPtr<ITmvMediaMuxerFactory, ESPMode::ThreadSafe>> MuxerFactories;
+		TmvMediaModule->GetMuxerFactories(MuxerFactories);
+		for (const TWeakPtr<ITmvMediaMuxerFactory, ESPMode::ThreadSafe>& FactoryWeak : MuxerFactories)
+		{
+			if (const TSharedPtr<ITmvMediaMuxerFactory, ESPMode::ThreadSafe> MuxerFactory = FactoryWeak.Pin())
+			{
+				return MuxerFactory->GetName();
+			}
+		}
+		return NAME_None;
+	}
 }
 
 void UTmvMediaTranscodeList::InsertItemAt(int32 InItemIndex, const FTmvMediaTranscodeListItem* InItem)
@@ -58,6 +81,7 @@ void UTmvMediaTranscodeList::InsertItemAt(int32 InItemIndex, const FTmvMediaTran
 	else
 	{
 		NewItem.Name = GenerateNewItemName();
+		NewItem.Settings.Muxer.Name = UE::TmvMedia::Transcode::GetDefaultMuxerName();
 		if (const UScriptStruct* SelectedStruct = UE::TmvMedia::Transcode::GetDefaultEncoderOptionsStruct())
 		{
 			NewItem.EncoderOptions.InitializeAsScriptStruct(SelectedStruct);
diff --git a/Engine/Plugins/Media/TmvMedia/Source/TmvMedia/Private/Transcoder/TmvMediaTranscodeMuxer.cpp b/Engine/Plugins/Media/TmvMedia/Source/TmvMedia/Private/Transcoder/TmvMediaTranscodeMuxer.cpp
--- a/Engine/Plugins/Media/TmvMedia/Source/TmvMedia/Private/Transcoder/TmvMediaTranscodeMuxer.cpp
+++ b/Engine/Plugins/Media/TmvMedia/Source/TmvMedia/Private/Transcoder/TmvMediaTranscodeMuxer.cpp
@@ -2,124 +2,58 @@
 
 #include "Transcoder/TmvMediaTranscodeMuxer.h"
 
-#include "GenericPlatform/GenericPlatformFile.h"
-#include "HAL/PlatformFileManager.h"
-#include "Misc/Paths.h"
+#include "ITmvMediaModule.h"
 #include "TmvMediaLog.h"
-#include "Transcoder/TmvMediaFrameMips.h"
-#include "Transcoder/TmvMediaFrameProducer.h"
+#include "Encoder/ITmvMediaMuxerFactory.h"
+#include "Misc/Paths.h"
 #include "Transcoder/TmvMediaTranscodeJob.h"
-#include "Utils/TmvMediaPathUtils.h"
 
 #include UE_INLINE_GENERATED_CPP_BY_NAME(TmvMediaTranscodeMuxer)
 
-#define LOCTEXT_NAMESPACE "TmvMediaTranscodeMuxer"
-
-int32 UTmvMediaTranscodeMuxer::OpenStream(UTmvMediaTranscodeJob* InParentJob, const FString& InStreamName, const FString& InExtension)
+FString UTmvMediaTranscodeMuxer::GetContainerOutputFilePath(const UTmvMediaTranscodeJob* InParentJob, const TSharedPtr<ITmvMediaMuxerFactory, ESPMode::ThreadSafe>& InMuxerFactory)
 {
-	if (!InParentJob)
+	if (!IsValid(InParentJob))
 	{
-		return INDEX_NONE;
+		return FString();
 	}
 
-	if (InStreamName.IsEmpty() || InExtension.IsEmpty())
-	{
-		const FText Message = FText::Format(LOCTEXT("InvalidStream", "Invalid stream name or extension (Name=\"{0}\", Ext=\"{1}\""),
-			FText::FromString(InStreamName), FText::FromString(InExtension));
-		UE_LOGF(LogTmvMedia, Error, "TranscodeMuxer::OpenStream: %ls.", *Message.ToString());
-		UTmvMediaTranscodeJob::SafeReportError(InParentJob, Message);
-		return INDEX_NONE;
-	}
-
-	const FString OutPath = InParentJob->Settings.GetAbsoluteOutputPath();
-
-	if (OutPath.IsEmpty())
-	{
-		const FText Message = LOCTEXT("EmptyOutputPath", "OutputPath is empty");
-		UE_LOGF(LogTmvMedia, Error, "TranscodeMuxer::OpenStream: %ls.", *Message.ToString());
-		UTmvMediaTranscodeJob::SafeReportError(InParentJob, Message);
-		return INDEX_NONE;
-	}
+	TSharedPtr<ITmvMediaMuxerFactory, ESPMode::ThreadSafe> MuxerFactory = InMuxerFactory;
 
-	if (!FPaths::DirectoryExists(OutPath))
+	if (!MuxerFactory)
 	{
-		if (!FPlatformFileManager::Get().GetPlatformFile().CreateDirectoryTree(*OutPath))
+		const FName MuxerName(InParentJob->Settings.Muxer.Name);
+		MuxerFactory = FindMuxerFactoryByName(MuxerName);
+		if (!MuxerFactory)
 		{
-			const FText Message = FText::Format(LOCTEXT("FailedCreateOutputDir", "Failed to create output directory \"{0}\""), FText::FromString(OutPath));
-			UE_LOGF(LogTmvMedia, Error, "TranscodeMuxer::OpenStream: %ls.", *Message.ToString());
-			UTmvMediaTranscodeJob::SafeReportError(InParentJob, Message);
-			return INDEX_NONE;
+			UE_LOGF(LogTmvMedia, Warning, "No Muxer Factory for \"%ls\".", *InParentJob->Settings.Muxer.Name.ToString());
 		}
 	}
 
-	FScopeLock Lock(&StreamsCS);
-	Streams.Add({InStreamName, InExtension});
-	return Streams.Num() - 1;
+	const FString OutPath = InParentJob->Settings.GetAbsoluteOutputPath();
+	const FString BaseName = !InParentJob->Settings.OutputBaseName.IsEmpty() ? InParentJob->Settings.OutputBaseName : InParentJob->JobName;
+	const FString FileExtension = MuxerFactory ? MuxerFactory->GetFileExtension() : TEXT("tmv");
+	return FPaths::Combine(OutPath, BaseName + TEXT(".") + FileExtension);
 }
 
-void UTmvMediaTranscodeMuxer::ReceiveAccessUnit(
-		UTmvMediaTranscodeJob* InParentJob,
-		int32 InStreamId,
-		const FTmvMediaFrameTimeInfo& InTimeInfo,
-		TSharedPtr<TArray64<uint8>>&& InAccessUnit)
+TSharedPtr<ITmvMediaMuxerFactory, ESPMode::ThreadSafe> UTmvMediaTranscodeMuxer::FindMuxerFactoryByName(const FName InFactoryName)
 {
-	if (!InAccessUnit.IsValid())
-	{
-		const FText Message = LOCTEXT("InvalidAccessUnit", "Null access unit");
-		UE_LOGF(LogTmvMedia, Error, "TranscodeMuxer::ReceiveAccessUnit: %ls.", *Message.ToString());
-		UTmvMediaTranscodeJob::SafeReportError(InParentJob, Message);
-		return; 
-	}
-
-	if (InAccessUnit->Num() == 0)
+	const ITmvMediaModule* TmvMediaModule = ITmvMediaModule::Get();
+	if (!TmvMediaModule)
 	{
-		UE_LOGF(LogTmvMedia, Warning, "ReceiveAccessUnit: Empty access unit (stream %d, frame %d).", InStreamId, InTimeInfo.FrameIndex);
-		return;
+		return nullptr;
 	}
-
-	FStreamInfo StreamInfo;
-	{
-		FScopeLock Lock(&StreamsCS);
-		if (!Streams.IsValidIndex(InStreamId))
-		{
-			const FText Message = LOCTEXT("InvalidStreamId", "Invalid stream id");
-			UE_LOGF(LogTmvMedia, Error, "TranscodeMuxer::ReceiveAccessUnit: %ls.", *Message.ToString());
-			UTmvMediaTranscodeJob::SafeReportError(InParentJob, Message);
-			return;
-		}
-		StreamInfo = Streams[InStreamId];
-	}
-
-	if (!InParentJob)
-	{
-		UE_LOGF(LogTmvMedia, Error, "ReceiveAccessUnit: Invalid Transcode Job.");
-		return;
-	}
-
-	// Remark: for a file sequence muxer, we don't need to synchronize (no common container) and
-	// we can just write the file immediately.
-	// For a container muxer, we need to keep the access units and re-order them. They can be
-	// received out of order because of multiple worker threads. We would need to know the number of worker threads.
-	const FString OutPath = InParentJob->Settings.GetAbsoluteOutputPath();
-	const FString FileName = UE::TmvMedia::PathUtils::MakeFrameFilename(*StreamInfo.StreamName, InTimeInfo.FrameIndex, InParentJob->Settings.ZeroPadFrameNumbers, *StreamInfo.Extension);
-	const FString FullFileName = FPaths::Combine(OutPath, FileName);
-
-	TUniquePtr<IFileHandle> File(FPlatformFileManager::Get().GetPlatformFile().OpenWrite(*FullFileName));
-	if (File)
+	
+	TArray<TWeakPtr<ITmvMediaMuxerFactory, ESPMode::ThreadSafe>> MuxerFactories;
+	TmvMediaModule->GetMuxerFactories(MuxerFactories);
+	for (const TWeakPtr<ITmvMediaMuxerFactory, ESPMode::ThreadSafe>& FactoryWeak : MuxerFactories)
 	{
-		if (!File->Write(InAccessUnit->GetData(), InAccessUnit->Num()))
+		if (TSharedPtr<ITmvMediaMuxerFactory, ESPMode::ThreadSafe> MuxerFactory = FactoryWeak.Pin())
 		{
-			const FText Message = FText::Format(LOCTEXT("WriteFailed", "Failed to write access unit to file \"{0}\""), FText::FromString(FullFileName));
-			UE_LOGF(LogTmvMedia, Error, "TranscodeMuxer::ReceiveAccessUnit: %ls.", *Message.ToString());
-			UTmvMediaTranscodeJob::SafeReportError(InParentJob, Message);
+			if (MuxerFactory->GetName() == InFactoryName)
+			{
+				return MuxerFactory;
+			}
 		}
 	}
-	else
-	{
-		const FText Message = FText::Format(LOCTEXT("OpenFailed", "Failed to open output file for writing: \"{0}\""), FText::FromString(FullFileName));
-		UE_LOGF(LogTmvMedia, Error, "TranscodeMuxer::ReceiveAccessUnit: %ls.", *Message.ToString());
-		UTmvMediaTranscodeJob::SafeReportError(InParentJob, Message);
-	}
+	return nullptr;
 }
-
-#undef LOCTEXT_NAMESPACE
\ No newline at end of file
diff --git a/Engine/Plugins/Media/TmvMedia/Source/TmvMedia/Public/Decoder/ITmvMediaDemuxer.h b/Engine/Plugins/Media/TmvMedia/Source/TmvMedia/Public/Decoder/ITmvMediaDemuxer.h
new file mode 100644
--- /dev/null
+++ b/Engine/Plugins/Media/TmvMedia/Source/TmvMedia/Public/Decoder/ITmvMediaDemuxer.h
@@ -0,0 +1,140 @@
+// Copyright Epic Games, Inc. All Rights Reserved.
+
+#pragma once
+
+#include "Containers/ArrayView.h"
+#include "Misc/Timespan.h"
+#include "Templates/SharedPointer.h"
+#include "TmvMediaTrackInfo.h"
+
+
+/**
+ * Information about a single track in a demuxed container.
+ * Extends the common track descriptor with demuxer-specific fields.
+ */
+struct FTmvMediaDemuxerTrackInfo : public FTmvMediaTrackInfo
+{
+	/** Track index (0-based). */
+	int32 TrackIndex = -1;
+
+	/** Total duration of the track in its timescale units. */
+	int64 Duration = 0;
+
+	/** Total number of samples in this track. */
+	uint32 NumSamples = 0;
+
+	/** Language code (3-letter ISO 639-2T). */
+	FString LanguageCode;
+};
+
+/**
+ * A single sample read from a demuxed container.
+ * Extends the common sample metadata with demuxer-specific fields.
+ */
+struct FTmvMediaDemuxerSample : public FTmvMediaTrackSampleInfo
+{
+	/** Raw coded sample data. */
+	TArray64<uint8> Data;
+
+	/** File byte offset of this sample. */
+	int64 FileOffset = 0;
+
+	/** Size in bytes (matches Data.Num() when data is populated). */
+	int64 SampleSize = 0;
+};
+
+/**
+ * Abstract interface for demuxing media samples from a container.
+ *
+ * Workflow:
+ *   1. Open() the container (from file path or from memory buffer).
+ *   2. GetTrackCount() / GetTrackInfo() to enumerate available tracks.
+ *   3. ReadSample() to iterate through samples on a given track.
+ *   4. Seek() to jump to a particular time or keyframe.
+ *   5. Close() when done.
+ *
+ * The demuxer extracts raw coded samples (access units) from the container
+ * that can then be fed into ITmvMediaDecoder for decoding.
+ *
+ * Thread safety: implementations are NOT internally synchronized. Callers must
+ * serialize all calls to the same demuxer instance (e.g. with a mutex).
+ */
+class ITmvMediaDemuxer
+{
+public:
+	virtual ~ITmvMediaDemuxer() = default;
+
+	/**
+	 * Open a container from a file path.
+	 * Parses the container structure and prepares track metadata.
+	 * @param InFilePath Absolute path to the container file.
+	 * @return Success if the container was parsed successfully.
+	 */
+	virtual ETmvMediaContainerResult OpenFile(const FString& InFilePath) = 0;
+
+	/**
+	 * Open a container from an in-memory buffer.
+	 * @param InData Buffer containing the complete container data.
+	 * @return Success if the container was parsed successfully.
+	 */
+	virtual ETmvMediaContainerResult OpenBuffer(TConstArrayView<uint8> InData) = 0;
+
+	/**
+	 * Returns the number of tracks found in the container.
+	 */
+	virtual int32 GetTrackCount() const = 0;
+
+	/**
+	 * Retrieve metadata for a specific track.
+	 * @param InTrackIndex Track index (0-based, must be < GetTrackCount()).
+	 * @param OutTrackInfo Populated track information.
+	 * @return Success if the track info was retrieved.
+	 */
+	virtual ETmvMediaContainerResult GetTrackInfo(int32 InTrackIndex, FTmvMediaDemuxerTrackInfo& OutTrackInfo) const = 0;
+
+	/**
+	 * Read the next sample from the specified track.
+	 * Advances the internal read position for that track.
+	 * @param InTrackIndex Track index.
+	 * @param OutSample Populated with the sample data and metadata.
+	 * @return Success if a sample was read, EndOfStream if no more samples.
+	 */
+	virtual ETmvMediaContainerResult ReadSample(int32 InTrackIndex, FTmvMediaDemuxerSample& OutSample) = 0;
+
+	/**
+	 * Read only the sample metadata (no data copy) from the specified track.
+	 * Useful for building sample tables or scanning without loading sample data.
+	 * @param InTrackIndex Track index.
+	 * @param OutSample Populated with metadata (Data array will be empty, SampleSize set).
+	 * @return Success if metadata was read, EndOfStream if no more samples.
+	 */
+	virtual ETmvMediaContainerResult ReadSampleInfo(int32 InTrackIndex, FTmvMediaDemuxerSample& OutSample) = 0;
+
+	/**
+	 * Seek to a keyframe at or before the given time.
+	 * After seeking, the next ReadSample call will return the keyframe sample.
+	 * @param InTrackIndex Track index.
+	 * @param InTime Target time.
+	 * @param InLaterTimeThreshold If a keyframe is found within this threshold after InTime, prefer it.
+	 * @return Success if a valid seek point was found.
+	 */
+	virtual ETmvMediaContainerResult Seek(int32 InTrackIndex, FTimespan InTime, FTimespan InLaterTimeThreshold = FTimespan::Zero()) = 0;
+
+	/**
+	 * Seek to a specific sample number.
+	 * @param InTrackIndex Track index.
+	 * @param InSampleNumber Target sample number (0-based).
+	 * @return Success if the sample number is valid.
+	 */
+	virtual ETmvMediaContainerResult SeekToSample(int32 InTrackIndex, uint32 InSampleNumber) = 0;
+
+	/**
+	 * Close the container and release resources.
+	 */
+	virtual void Close() = 0;
+
+	/**
+	 * Returns a human-readable error message if the last operation failed.
+	 */
+	virtual FString GetLastError() const = 0;
+};
diff --git a/Engine/Plugins/Media/TmvMedia/Source/TmvMedia/Public/Decoder/ITmvMediaDemuxerFactory.h b/Engine/Plugins/Media/TmvMedia/Source/TmvMedia/Public/Decoder/ITmvMediaDemuxerFactory.h
new file mode 100644
--- /dev/null
+++ b/Engine/Plugins/Media/TmvMedia/Source/TmvMedia/Public/Decoder/ITmvMediaDemuxerFactory.h
@@ -0,0 +1,38 @@
+// Copyright Epic Games, Inc. All Rights Reserved.
+
+#pragma once
+
+#include "Containers/Array.h"
+#include "Containers/UnrealString.h"
+#include "Internationalization/Text.h"
+#include "Templates/SharedPointer.h"
+#include "UObject/NameTypes.h"
+
+class ITmvMediaDemuxer;
+
+/**
+ * Interface for an ITmvMediaDemuxer factory.
+ */
+class ITmvMediaDemuxerFactory
+{
+public:
+	virtual ~ITmvMediaDemuxerFactory() = default;
+
+	/** Returns the name of this demuxer implementation. */
+	virtual FName GetName() const = 0;
+
+	/** Returns the display name of this demuxer implementation. */
+	virtual FText GetDisplayName() const = 0;
+
+	/**
+	 * Return the list of container format identifiers this factory supports.
+	 * Ex: "mp4", "fmp4", "mov"
+	 */
+	virtual TArray<FString> GetSupportedContainerFormats() const = 0;
+
+	/**
+	 * Creates a demuxer instance.
+	 * @return Created demuxer, or null if failed.
+	 */
+	virtual TSharedPtr<ITmvMediaDemuxer, ESPMode::ThreadSafe> CreateDemuxer() = 0;
+};
diff --git a/Engine/Plugins/Media/TmvMedia/Source/TmvMedia/Public/Encoder/ITmvMediaMuxer.h b/Engine/Plugins/Media/TmvMedia/Source/TmvMedia/Public/Encoder/ITmvMediaMuxer.h
new file mode 100644
--- /dev/null
+++ b/Engine/Plugins/Media/TmvMedia/Source/TmvMedia/Public/Encoder/ITmvMediaMuxer.h
@@ -0,0 +1,175 @@
+// Copyright Epic Games, Inc. All Rights Reserved.
+
+#pragma once
+
+#include "Containers/ArrayView.h"
+#include "Delegates/Delegate.h"
+#include "Misc/Optional.h"
+#include "Misc/Timespan.h"
+#include "Templates/SharedPointer.h"
+#include "TmvMediaTrackInfo.h"
+
+
+/**
+ * Configuration for a single track in the muxer.
+ * Extends the common track descriptor with muxer-specific fields.
+ */
+struct FTmvMediaMuxerTrackConfig : public FTmvMediaTrackInfo
+{
+	/**
+	 * If true, every sample in this track is independently decodable (sync sample).
+	 * If false, samples may depend on other samples.
+	 */
+	bool bIsAllKeyframes = true;
+
+	/** Optional compressor name string (video tracks only). */
+	FString CompressorName;
+
+	/**
+	 * If set, enables subsample information for each sample.
+	 * The value is the 24-bit flags for the subsample box.
+	 */
+	TOptional<uint32> SubSampleFlags;
+};
+
+/**
+ * Configuration for the muxer output.
+ */
+struct FTmvMediaMuxerConfig
+{
+	/** Output mode for the container. */
+	enum class EOutputMode : uint8
+	{
+		/** Standard mode. Container metadata at end. */
+		Standard,
+		/** Web-optimized mode. Container metadata at front (requires temporary file + 2nd pass). */
+		WebOptimized,
+		/** Fragmented mode. Metadata interleaved with data. */
+		Fragmented
+	};
+
+	/** Output file path. */
+	FString OutputFilename;
+
+	/** Temporary file path (required for WebOptimized mode). */
+	FString TemporaryFilename;
+
+	/** Output mode. */
+	EOutputMode OutputMode = EOutputMode::Standard;
+
+	/**
+	 * Duration for interleaving chunks of sample data.
+	 * Samples are collected per track until this duration is reached, then written.
+	 * Typical value: 500ms.
+	 */
+	FTimespan InterleaveDuration = FTimespan::FromMilliseconds(500);
+};
+
+/**
+ * A single sample to be muxed into a track.
+ * Extends the common sample metadata with muxer-specific fields.
+ */
+struct FTmvMediaMuxerSample : public FTmvMediaTrackSampleInfo
+{
+	/** Subsample information entry. */
+	struct FSubSampleInfo
+	{
+		uint32 CodecSpecificParameters = 0;
+		uint32 SubSampleSize = 0;
+		uint8 SubSamplePriority = 0;
+		uint8 Discardable = 0;
+	};
+
+	/** Raw coded sample data. Written verbatim to the container. */
+	TConstArrayView64<uint8> Data;
+
+	/** Whether this is the final sample for this track. */
+	bool bIsFinalSample = false;
+
+	/** Optional subsample information. */
+	TArray<FSubSampleInfo> SubSamples;
+};
+
+/**
+ * Abstract interface for muxing encoded access units into a media container.
+ *
+ * Workflow:
+ *   1. Configure() with output settings.
+ *   2. AddTrack() for each media track, receiving a track index.
+ *   3. Start() to begin the muxing process.
+ *   4. AddSample() for each sample on each track (called from the sample request callback).
+ *   5. Finalize() once all tracks signal their final sample.
+ *
+ * Sample delivery is asynchronous: the muxer calls back via FSampleRequestDelegate
+ * to request samples. The caller provides them via AddSample().
+ */
+class ITmvMediaMuxer
+{
+public:
+	virtual ~ITmvMediaMuxer() = default;
+
+	/** Defines an entry in the sample request array indicating which track needs a sample. */
+	struct FSampleRequest
+	{
+		int32 TrackIndex = -1;
+		uint32 SampleNumber = 0;
+	};
+
+	/** Delegate invoked when the muxer needs samples for one or more tracks. */
+	DECLARE_DELEGATE_OneParam(FSampleRequestDelegate, const TArray<FSampleRequest>& /*PendingRequests*/);
+
+	/** Status reported upon completion or failure. */
+	enum class EStatus
+	{
+		Failed,
+		Finished
+	};
+
+	/** Delegate invoked when muxing completes or fails. */
+	DECLARE_DELEGATE_OneParam(FStatusDelegate, EStatus /*Status*/);
+
+	/**
+	 * Configure the muxer with output settings.
+	 * Must be called before AddTrack().
+	 * @param InConfig Muxer configuration.
+	 * @return Success if configuration is valid, Fail otherwise.
+	 */
+	virtual ETmvMediaContainerResult Configure(const FTmvMediaMuxerConfig& InConfig) = 0;
+
+	/**
+	 * Add a track to the container.
+	 * Must be called after Configure() and before Start().
+	 * @param InTrackConfig Track configuration.
+	 * @return Track index (>= 0) on success, or -1 on failure.
+	 */
+	virtual int32 AddTrack(const FTmvMediaMuxerTrackConfig& InTrackConfig) = 0;
+
+	/**
+	 * Start the muxing process.
+	 * The muxer will begin calling InSampleRequestDelegate to request samples.
+	 * @param InSampleRequestDelegate Called when the muxer needs more sample data.
+	 * @param InStatusDelegate Called when muxing completes or fails.
+	 * @return Success if started, Fail otherwise.
+	 */
+	virtual ETmvMediaContainerResult Start(FSampleRequestDelegate InSampleRequestDelegate, FStatusDelegate InStatusDelegate) = 0;
+
+	/**
+	 * Provide a sample to the muxer for the specified track.
+	 * Called in response to FSampleRequestDelegate.
+	 * @param InTrackIndex Index returned by AddTrack().
+	 * @param InSample The sample data and metadata.
+	 * @return Success if the sample was accepted.
+	 */
+	virtual ETmvMediaContainerResult AddSample(int32 InTrackIndex, const FTmvMediaMuxerSample& InSample) = 0;
+
+	/**
+	 * Stop muxing and finalize the output file.
+	 * @return Success if the file was closed properly.
+	 */
+	virtual ETmvMediaContainerResult Finalize() = 0;
+
+	/**
+	 * Returns a human-readable error message if the last operation failed.
+	 */
+	virtual FString GetLastError() const = 0;
+};
diff --git a/Engine/Plugins/Media/TmvMedia/Source/TmvMedia/Public/Encoder/ITmvMediaMuxerFactory.h b/Engine/Plugins/Media/TmvMedia/Source/TmvMedia/Public/Encoder/ITmvMediaMuxerFactory.h
new file mode 100644
--- /dev/null
+++ b/Engine/Plugins/Media/TmvMedia/Source/TmvMedia/Public/Encoder/ITmvMediaMuxerFactory.h
@@ -0,0 +1,44 @@
+// Copyright Epic Games, Inc. All Rights Reserved.
+
+#pragma once
+
+#include "Containers/Array.h"
+#include "Containers/UnrealString.h"
+#include "Internationalization/Text.h"
+#include "Templates/SharedPointer.h"
+#include "UObject/NameTypes.h"
+
+class ITmvMediaMuxer;
+
+/**
+ * Interface for an ITmvMediaMuxer factory.
+ */
+class ITmvMediaMuxerFactory
+{
+public:
+	virtual ~ITmvMediaMuxerFactory() = default;
+
+	/** Returns the name of this muxer implementation. */
+	virtual FName GetName() const = 0;
+
+	/** Returns the display name of this muxer implementation. */
+	virtual FText GetDisplayName() const = 0;
+
+	/**
+	 * Return the list of container format identifiers this factory supports.
+	 * Ex: "tmv"
+	 */
+	virtual TArray<FString> GetSupportedContainerFormats() const = 0;
+
+	/**
+	 * Returns the file extension for the output container (without leading dot).
+	 * Ex: "tmv"
+	 */
+	virtual FString GetFileExtension() const = 0;
+
+	/**
+	 * Creates a muxer instance.
+	 * @return Created muxer, or null if failed.
+	 */
+	virtual TSharedPtr<ITmvMediaMuxer, ESPMode::ThreadSafe> CreateMuxer() = 0;
+};
diff --git a/Engine/Plugins/Media/TmvMedia/Source/TmvMedia/Public/Encoder/TmvMediaEncoderOptions.h b/Engine/Plugins/Media/TmvMedia/Source/TmvMedia/Public/Encoder/TmvMediaEncoderOptions.h
--- a/Engine/Plugins/Media/TmvMedia/Source/TmvMedia/Public/Encoder/TmvMediaEncoderOptions.h
+++ b/Engine/Plugins/Media/TmvMedia/Source/TmvMedia/Public/Encoder/TmvMediaEncoderOptions.h
@@ -67,7 +67,7 @@ enum class ETmvMediaEncoderColorMatrixRange : uint8
 /**
  * Encoder options.
  */
-USTRUCT(BlueprintType)
+USTRUCT()
 struct FTmvMediaEncoderOptions
 {
 	GENERATED_BODY()
@@ -78,6 +78,9 @@ struct FTmvMediaEncoderOptions
 	/** Returns the extension used for saving to file sequence. */
 	virtual FString GetFileSequenceExtension() const { return FString(); }
 
+	/** Returns the FOURCC sample entry format identifying the codec (e.g. 'apv1'). Used for container muxing. */
+	virtual uint32 GetCodecFourCC() const { return 0; }
+
 	virtual ~FTmvMediaEncoderOptions() = default;
 
 	/** 
diff --git a/Engine/Plugins/Media/TmvMedia/Source/TmvMedia/Public/ITmvMediaModule.h b/Engine/Plugins/Media/TmvMedia/Source/TmvMedia/Public/ITmvMediaModule.h
--- a/Engine/Plugins/Media/TmvMedia/Source/TmvMedia/Public/ITmvMediaModule.h
+++ b/Engine/Plugins/Media/TmvMedia/Source/TmvMedia/Public/ITmvMediaModule.h
@@ -10,7 +10,9 @@
 #define UE_API TMVMEDIA_API
 
 class ITmvMediaDecoderFactory;
+class ITmvMediaDemuxerFactory;
 class ITmvMediaEncoderFactory;
+class ITmvMediaMuxerFactory;
 class ITmvMediaTranscodeJobManager;
 
 /**
@@ -26,6 +28,7 @@ public:
 
 	/**
 	 * Returns existing instance of the module, loads it if it was not already loaded.
+	 * Beware of calling this during the shutdown phase, though. The module might have been unloaded already.
 	 */
 	static UE_API ITmvMediaModule& GetOrLoad();
 	
@@ -42,7 +45,7 @@ public:
 	/**
 	 * Get a list of all currently registered factories. 
 	 */
-	virtual void GetDecoderFactories(TArray<TWeakPtr<ITmvMediaDecoderFactory, ESPMode::ThreadSafe>>& OutDecoderFactories) = 0;
+	virtual void GetDecoderFactories(TArray<TWeakPtr<ITmvMediaDecoderFactory, ESPMode::ThreadSafe>>& OutDecoderFactories) const = 0;
 
 	/**
 	 * Returns the best decoder factory for the given format.
@@ -50,7 +53,7 @@ public:
 	 * @param InOptions Additional options may be provided to specify the format more closely, which may help selecting the best suited implementation.
 	 * @return Highest priority decoder factory if found, null otherwise.
 	 */
-	virtual TSharedPtr<ITmvMediaDecoderFactory, ESPMode::ThreadSafe> GetBestDecoderFactoryForFormat(const FString& InCodecFormat, const TMap<FString, FVariant>& InOptions) = 0;
+	virtual TSharedPtr<ITmvMediaDecoderFactory, ESPMode::ThreadSafe> GetBestDecoderFactoryForFormat(const FString& InCodecFormat, const TMap<FString, FVariant>& InOptions) const = 0;
 
 	/**
 	 * Register an encoder factory in the module.
@@ -74,6 +77,43 @@ public:
 	 */
 	virtual TSharedPtr<ITmvMediaEncoderFactory, ESPMode::ThreadSafe> FindEncoderFactory(FName InEncoderFactoryName) const = 0;
 
+	/**
+	 * Register a muxer factory in the module.
+	 */
+	virtual void RegisterMuxerFactory(TSharedPtr<ITmvMediaMuxerFactory, ESPMode::ThreadSafe> InMuxerFactory) = 0;
+
+	/**
+	 * Unregister a muxer factory from the module.
+	 */
+	virtual void UnregisterMuxerFactory(TSharedPtr<ITmvMediaMuxerFactory, ESPMode::ThreadSafe> InMuxerFactory) = 0;
+
+	/**
+	 * Get a list of all currently registered muxer factories.
+	 */
+	virtual void GetMuxerFactories(TArray<TWeakPtr<ITmvMediaMuxerFactory, ESPMode::ThreadSafe>>& OutMuxerFactories) const = 0;
+
+	/**
+	 * Register a demuxer factory in the module.
+	 */
+	virtual void RegisterDemuxerFactory(TSharedPtr<ITmvMediaDemuxerFactory, ESPMode::ThreadSafe> InDemuxerFactory) = 0;
+
+	/**
+	 * Unregister a demuxer factory from the module.
+	 */
+	virtual void UnregisterDemuxerFactory(TSharedPtr<ITmvMediaDemuxerFactory, ESPMode::ThreadSafe> InDemuxerFactory) = 0;
+
+	/**
+	 * Get a list of all currently registered demuxer factories.
+	 */
+	virtual void GetDemuxerFactories(TArray<TWeakPtr<ITmvMediaDemuxerFactory, ESPMode::ThreadSafe>>& OutDemuxerFactories) const = 0;
+
+	/**
+	 * Find the demuxer factory that supports the given file extensions.
+	 * @param InFileExtension File extension without dot
+	 * @return Demuxer factory if found, null otherwise.
+	 */
+	virtual TSharedPtr<ITmvMediaDemuxerFactory, ESPMode::ThreadSafe> FindDemuxerFactoryForExtension(const FString& InFileExtension) const = 0;
+
 	/**
 	 * Access the global transcode job manager.
 	 */
diff --git a/Engine/Plugins/Media/TmvMedia/Source/TmvMedia/Public/TmvMediaTrackInfo.h b/Engine/Plugins/Media/TmvMedia/Source/TmvMedia/Public/TmvMediaTrackInfo.h
new file mode 100644
--- /dev/null
+++ b/Engine/Plugins/Media/TmvMedia/Source/TmvMedia/Public/TmvMediaTrackInfo.h
@@ -0,0 +1,97 @@
+// Copyright Epic Games, Inc. All Rights Reserved.
+
+#pragma once
+
+#include "Containers/Array.h"
+#include "Containers/UnrealString.h"
+
+/**
+ * Track type for muxer/demuxer tracks.
+ */
+enum class ETmvMediaTrackType : uint8
+{
+	Unspecified,
+	Video,
+	Audio
+};
+
+/**
+ * Common descriptor for a media track, shared by muxer and demuxer.
+ * Describes the codec, format, and media-type-specific properties of a track.
+ */
+struct FTmvMediaTrackInfo
+{
+	/** Track type. */
+	ETmvMediaTrackType TrackType = ETmvMediaTrackType::Unspecified;
+
+	/**
+	 * FOURCC sample entry format identifying the codec.
+	 * Ex: 'apv1' for APV, 'avc1' for H.264, 'hev1' for HEVC.
+	 */
+	uint32 SampleEntryFormat = 0;
+
+	/**
+	 * Number of ticks per second for this track's timing.
+	 * All DTS, PTS, and duration values are measured in this timescale.
+	 */
+	uint32 Timescale = 0;
+
+	/**
+	 * If every sample has the same duration, this holds the constant value in timescale units.
+	 * 0 means variable or unknown duration.
+	 * Frame rate can be computed as FFrameRate(Timescale, ConstantSampleDuration).
+	 */
+	uint32 ConstantSampleDuration = 0;
+
+	/** Display width (video tracks only). */
+	uint32 DisplayWidth = 0;
+
+	/** Display height (video tracks only). */
+	uint32 DisplayHeight = 0;
+
+	/** Sampling rate in Hz (audio tracks only). */
+	uint32 SamplingRate = 0;
+
+	/** Number of channels (audio tracks only). */
+	uint16 NumberOfChannels = 0;
+
+	/**
+	 * Codec-specific configuration boxes, pre-built as raw box data.
+	 * Each entry is a complete box (size + type + payload).
+	 */
+	TArray<TArray<uint8>> CodecSpecificBoxes;
+
+	/** Optional human-readable track name. */
+	FString TrackName;
+};
+
+/**
+ * Common result code for container muxer/demuxer operations.
+ */
+enum class ETmvMediaContainerResult
+{
+	Success,
+	Fail,
+	EndOfStream
+};
+
+/**
+ * Common sample metadata shared by muxer and demuxer sample structs.
+ */
+struct FTmvMediaTrackSampleInfo
+{
+	/** Presentation duration of this sample, in the track's timescale units. */
+	uint32 Duration = 0;
+
+	/** Decode timestamp in the track's timescale units. */
+	int64 DTS = 0;
+
+	/** Presentation (composition) timestamp in the track's timescale units. */
+	int64 PTS = 0;
+
+	/** Sequential sample number (0-based). */
+	uint32 SampleNumber = 0;
+
+	/** Whether this sample is a sync sample (keyframe). */
+	bool bIsKeyframe = false;
+};
diff --git a/Engine/Plugins/Media/TmvMedia/Source/TmvMedia/Public/Transcoder/TmvMediaFrameMips.h b/Engine/Plugins/Media/TmvMedia/Source/TmvMedia/Public/Transcoder/TmvMediaFrameMips.h
--- a/Engine/Plugins/Media/TmvMedia/Source/TmvMedia/Public/Transcoder/TmvMediaFrameMips.h
+++ b/Engine/Plugins/Media/TmvMedia/Source/TmvMedia/Public/Transcoder/TmvMediaFrameMips.h
@@ -17,6 +17,9 @@ struct FTmvMediaFrameTimeInfo
 {
 	/** Frame index (for index based muxers). */
 	int32 FrameIndex = 0;
+	
+	/** Zero-based Frame index (for index based muxers). */
+	int32 FrameIndexNoOffset = 0;
 
 	/** Frame time stamp (for time span based muxers). */
 	FTimespan FrameTime;
diff --git a/Engine/Plugins/Media/TmvMedia/Source/TmvMedia/Public/Transcoder/TmvMediaTranscodeJob.h b/Engine/Plugins/Media/TmvMedia/Source/TmvMedia/Public/Transcoder/TmvMediaTranscodeJob.h
--- a/Engine/Plugins/Media/TmvMedia/Source/TmvMedia/Public/Transcoder/TmvMediaTranscodeJob.h
+++ b/Engine/Plugins/Media/TmvMedia/Source/TmvMedia/Public/Transcoder/TmvMediaTranscodeJob.h
@@ -42,6 +42,31 @@ enum class ETmvMediaTranscodeInputSource : uint8
 	MediaSource
 };
 
+/**
+ * Selection of the output format.
+ */
+UENUM(BlueprintType)
+enum class ETmvMediaTranscodeOutputFormat : uint8
+{
+	/** Output each access unit as a separate file. */
+	FileSequence UMETA(DisplayName="TMV Image Sequence"),
+	/** Mux access units into a single container file. */
+	Container UMETA(DisplayName="TMV Video File"),
+};
+
+/** Structure for muxer settings used to select the muxer. */
+USTRUCT(BlueprintType)
+struct FTmvMediaTranscodeMuxerSettings
+{
+	GENERATED_BODY()
+
+	/**
+	 * The name of the selected muxer
+	 */
+	UPROPERTY(EditAnywhere, BlueprintReadWrite, Category = Settings)
+	FName Name;
+};
+
 /**
  * Common Transcode Job Settings
  */
@@ -94,6 +119,13 @@ struct FTmvMediaTranscodeJobSettings
 	UPROPERTY(EditAnywhere, Category = "Input Sequence", meta = (EditCondition = "InputSource ==  ETmvMediaTranscodeInputSource::MediaSource", EditConditionHides))
 	TSoftObjectPtr<UMediaSource> InputMediaSource;
 
+	/** Select the output format. */
+	UPROPERTY(EditAnywhere, Category = "Output Sequence", meta = (SegmentedDisplay))
+	ETmvMediaTranscodeOutputFormat OutputFormat = ETmvMediaTranscodeOutputFormat::FileSequence;
+	
+	UPROPERTY(EditAnywhere, Category = "Output Sequence", AdvancedDisplay, meta = (EditCondition = "OutputFormat == ETmvMediaTranscodeOutputFormat::Container", EditConditionHides))
+	FTmvMediaTranscodeMuxerSettings Muxer;
+
 	/** The directory to output the processed image sequence files to. */
 	UPROPERTY(EditAnywhere, Category = "Output Sequence")
 	FDirectoryPath OutputPath;
@@ -103,7 +135,7 @@ struct FTmvMediaTranscodeJobSettings
 	FString OutputBaseName;
 
 	/** Specify the zero padding on the frame number in the output sequence files. */
-	UPROPERTY(EditAnywhere, Category = "Output Sequence", AdvancedDisplay, meta=(ClampMin=0, UIMin=0))
+	UPROPERTY(EditAnywhere, Category = "Output Sequence", AdvancedDisplay, meta=(ClampMin=0, UIMin=0, EditCondition = "OutputFormat == ETmvMediaTranscodeOutputFormat::FileSequence", EditConditionHides))
 	int32 ZeroPadFrameNumbers = 5;
 
 	/** If enabled, a media source asset will be made for the output sequence. */
@@ -312,11 +344,12 @@ public:
 	{
 		if (IsValid(InJob))
 		{
+			const TSharedPtr<ITmvMediaTranscodeNotification> Notification = InJob->Notification;
 			InJob->bIsError = true;
 
-			if (!InErrorMessage.IsEmpty() && InJob->Notification.IsValid())
+			if (!InErrorMessage.IsEmpty() && Notification.IsValid())
 			{
-				InJob->Notification->SetText(InErrorMessage);
+				Notification->SetText(InErrorMessage);
 			}
 		}
 	}
@@ -393,7 +426,7 @@ public:
 	 * It is driven by control functions like Start, Stop and then synchronized with the
 	 * stage's state in the Tick function.
 	 */
-	TAtomic<ETmvMediaTranscodeJobStatus> JobStatus = ETmvMediaTranscodeJobStatus::None;
+	std::atomic<ETmvMediaTranscodeJobStatus> JobStatus = ETmvMediaTranscodeJobStatus::None;
 
 	/** Indicate a stage had an error somewhere that requires the job to stop. */
 	std::atomic<bool> bIsError = false;
diff --git a/Engine/Plugins/Media/TmvMedia/Source/TmvMedia/Public/Transcoder/TmvMediaTranscodeMuxer.h b/Engine/Plugins/Media/TmvMedia/Source/TmvMedia/Public/Transcoder/TmvMediaTranscodeMuxer.h
--- a/Engine/Plugins/Media/TmvMedia/Source/TmvMedia/Public/Transcoder/TmvMediaTranscodeMuxer.h
+++ b/Engine/Plugins/Media/TmvMedia/Source/TmvMedia/Public/Transcoder/TmvMediaTranscodeMuxer.h
@@ -9,37 +9,51 @@
 
 #define UE_API TMVMEDIA_API
 
+class ITmvMediaMuxerFactory;
 struct FTmvMediaFrameTimeInfo;
 
 /**
- * Base class for muxer stage of transcode pipeline.
+ * Abstract base class for muxer stage of transcode pipeline.
  * @see UTmvMediaTranscodeJob
  */
-UCLASS(MinimalAPI)
+UCLASS(MinimalAPI, Abstract)
 class UTmvMediaTranscodeMuxer : public UTmvMediaTranscodeStage
 {
 	GENERATED_BODY()
 public:
 	/**
 	 * Opens a stream.
-	 * 
+	 *
 	 * @param InParentJob Parent Job this stage is part of.
 	 * @param InStreamName Name given to the stream (will be used for filename prefix in case of file sequence muxer).
 	 * @param InExtension Extension for the filename in case of file sequence muxer.
 	 * @return Index of the stream or INDEX_NONE in case of error.
-	 * 
-	 * todo: api only for file sequence muxer. Will be revisited for containers.
 	 */
-	UE_API virtual int32 OpenStream(UTmvMediaTranscodeJob* InParentJob, const FString& InStreamName, const FString& InExtension);
+	virtual int32 OpenStream(UTmvMediaTranscodeJob* InParentJob, const FString& InStreamName, const FString& InExtension) PURE_VIRTUAL(UTmvMediaTranscodeMuxer::OpenStream, return INDEX_NONE;);
 
 	/**
-	 * Push an access unit buffer in the stream at the given index (must call OpenStream first). 
+	 * Push an access unit buffer in the stream at the given index (must call OpenStream first).
 	 */
-	UE_API virtual void ReceiveAccessUnit(
+	virtual void ReceiveAccessUnit(
 		UTmvMediaTranscodeJob* InParentJob,
 		int32 InStreamId,
 		const FTmvMediaFrameTimeInfo& TimeInfo,
-		TSharedPtr<TArray64<uint8>>&& InAccessUnit); // todo: improve abstraction to bridge different implementations.
+		TSharedPtr<TArray64<uint8>>&& InAccessUnit) PURE_VIRTUAL(UTmvMediaTranscodeMuxer::ReceiveAccessUnit,);
+
+	/**
+	 * Utility function to get the absolute container path from job settings.
+	 * @param InParentJob Parent job to get settings from.
+	 * @param InMuxerFactory Optional muxer factory. If not specified, the factory is retrieve from job settings.
+	 * @return Absolute container path, including the extension.
+	 */
+	UE_API static FString GetContainerOutputFilePath(const UTmvMediaTranscodeJob* InParentJob, const TSharedPtr<ITmvMediaMuxerFactory, ESPMode::ThreadSafe>& InMuxerFactory = nullptr);
+
+	/**
+	 * Muxer Utility function to find a muxer factory by name.
+	 * @param InFactoryName Name of the muxer factory to find.
+	 * @return muxer factory if found, nullptr otherwise.
+	 */
+	UE_API static TSharedPtr<ITmvMediaMuxerFactory, ESPMode::ThreadSafe> FindMuxerFactoryByName(FName InFactoryName);
 
 protected:
 	/** Stream information */
@@ -48,7 +62,7 @@ protected:
 		FString StreamName;
 		FString Extension;
 	};
-	
+
 	/** Array of currently opened streams. */
 	TArray<FStreamInfo> Streams;
 
diff --git a/Engine/Plugins/Media/TmvMedia/Source/TmvMedia/Public/Transcoder/TmvMediaTranscodeStage.h b/Engine/Plugins/Media/TmvMedia/Source/TmvMedia/Public/Transcoder/TmvMediaTranscodeStage.h
--- a/Engine/Plugins/Media/TmvMedia/Source/TmvMedia/Public/Transcoder/TmvMediaTranscodeStage.h
+++ b/Engine/Plugins/Media/TmvMedia/Source/TmvMedia/Public/Transcoder/TmvMediaTranscodeStage.h
@@ -68,7 +68,7 @@ public:
 	}
 
 	/** Returns the current status of the stage. */
-	ETmvMediaTranscodeStageStatus GetStageStatus() const { return StageStatus.Load(); }
+	ETmvMediaTranscodeStageStatus GetStageStatus() const { return StageStatus.load(); }
 
 	/** Change the status of the stage. This will propagate job events. */
 	UE_API void SetStageStatus(ETmvMediaTranscodeStageStatus InStatus, UTmvMediaTranscodeJob* InParentJob);
@@ -85,7 +85,7 @@ protected:
 	
 private:
 	/** Indicate the status of the stage. */
-	TAtomic<ETmvMediaTranscodeStageStatus> StageStatus = ETmvMediaTranscodeStageStatus::Stopped;
+	std::atomic<ETmvMediaTranscodeStageStatus> StageStatus = ETmvMediaTranscodeStageStatus::Stopped;
 };
 
 #undef UE_API
\ No newline at end of file
diff --git a/Engine/Plugins/Media/TmvMedia/Source/TmvMedia/Public/Utils/TmvMediaUtils.h b/Engine/Plugins/Media/TmvMedia/Source/TmvMedia/Public/Utils/TmvMediaUtils.h
--- a/Engine/Plugins/Media/TmvMedia/Source/TmvMedia/Public/Utils/TmvMediaUtils.h
+++ b/Engine/Plugins/Media/TmvMedia/Source/TmvMedia/Public/Utils/TmvMediaUtils.h
@@ -62,4 +62,25 @@ namespace UE::TmvMedia::Utils
 		}
 		return NSLOCTEXT("TmvMediaUtils", "InvalidEnumText", "<invalid-enum>");
 	}
-}
\ No newline at end of file
+	
+	/** 
+	 * Constructs a FourCC integer from four 8-bit character codes, commonly used for 
+	 * identifying file formats, codec types, or pixel formats.
+	 */
+	inline constexpr uint32 MakeFourCC(const uint8 A, const uint8 B, const uint8 C, const uint8 D)
+	{
+		return (static_cast<uint32>(A) << 24) | (static_cast<uint32>(B) << 16) | (static_cast<uint32>(C) << 8) | static_cast<uint32>(D);
+	}
+
+	/**
+	 * Returns the string made out of the 4 character codes from the given FourCC value. 
+	 */
+	inline FString FourCCToString(uint32 InFourCC)
+	{
+		return FString::Printf(TEXT("%c%c%c%c"),
+			static_cast<TCHAR>((InFourCC >> 24) & 0xFF),
+			static_cast<TCHAR>((InFourCC >> 16) & 0xFF),
+			static_cast<TCHAR>((InFourCC >> 8) & 0xFF),
+			static_cast<TCHAR>(InFourCC & 0xFF));
+	}
+}
diff --git a/Engine/Plugins/Media/TmvMedia/Source/TmvMediaEditor/Private/Customizations/TmvMediaMuxerSettingsCustomization.cpp b/Engine/Plugins/Media/TmvMedia/Source/TmvMediaEditor/Private/Customizations/TmvMediaMuxerSettingsCustomization.cpp
new file mode 100644
--- /dev/null
+++ b/Engine/Plugins/Media/TmvMedia/Source/TmvMediaEditor/Private/Customizations/TmvMediaMuxerSettingsCustomization.cpp
@@ -0,0 +1,128 @@
+// Copyright Epic Games, Inc. All Rights Reserved.
+
+#include "TmvMediaMuxerSettingsCustomization.h"
+
+#include "DetailWidgetRow.h"
+#include "Encoder/ITmvMediaMuxerFactory.h"
+#include "IDetailChildrenBuilder.h"
+#include "ITmvMediaModule.h"
+#include "Widgets/Input/SComboBox.h"
+#include "Widgets/Text/STextBlock.h"
+
+#define LOCTEXT_NAMESPACE "TmvMediaMuxerSettingsCustomization"
+
+TSharedRef<IPropertyTypeCustomization> FTmvMediaMuxerSettingsCustomization::MakeInstance()
+{
+	return MakeShareable(new FTmvMediaMuxerSettingsCustomization());
+}
+
+void FTmvMediaMuxerSettingsCustomization::CustomizeHeader(
+	TSharedRef<IPropertyHandle> InPropertyHandle,
+	FDetailWidgetRow& InHeaderRow,
+	IPropertyTypeCustomizationUtils& InCustomizationUtils)
+{
+	NameProperty = InPropertyHandle->GetChildHandle(TEXT("Name"));
+
+	RefreshMuxerOptions();
+
+	InHeaderRow
+		.NameContent()
+		[
+			InPropertyHandle->CreatePropertyNameWidget()
+		]
+		.ValueContent()
+		.MinDesiredWidth(200.0f)
+		[
+			SNew(SComboBox<TSharedPtr<FString>>)
+			.OptionsSource(&MuxerDisplayNames)
+			.OnSelectionChanged(this, &FTmvMediaMuxerSettingsCustomization::OnMuxerSelectionChanged)
+			.OnGenerateWidget_Lambda([](TSharedPtr<FString> InItem)
+			{
+				return SNew(STextBlock).Text(FText::FromString(*InItem));
+			})
+			.Content()
+			[
+				SNew(STextBlock).Text(this, &FTmvMediaMuxerSettingsCustomization::GetSelectedMuxerDisplayName)
+			]
+		];
+}
+
+void FTmvMediaMuxerSettingsCustomization::CustomizeChildren(
+	TSharedRef<IPropertyHandle> InPropertyHandle,
+	IDetailChildrenBuilder& InChildBuilder,
+	IPropertyTypeCustomizationUtils& InCustomizationUtils)
+{
+	// No children — the Name field is fully handled in the header combo box.
+}
+
+void FTmvMediaMuxerSettingsCustomization::RefreshMuxerOptions()
+{
+	MuxerDisplayNames.Empty();
+	DisplayNameToFactoryName.Empty();
+
+	ITmvMediaModule* TmvModule = ITmvMediaModule::Get();
+	if (!TmvModule)
+	{
+		return;
+	}
+
+	TArray<TWeakPtr<ITmvMediaMuxerFactory, ESPMode::ThreadSafe>> MuxerFactories;
+	TmvModule->GetMuxerFactories(MuxerFactories);
+
+	for (const TWeakPtr<ITmvMediaMuxerFactory, ESPMode::ThreadSafe>& FactoryWeak : MuxerFactories)
+	{
+		if (TSharedPtr<ITmvMediaMuxerFactory, ESPMode::ThreadSafe> Factory = FactoryWeak.Pin())
+		{
+			const FString DisplayName = Factory->GetDisplayName().ToString();
+			MuxerDisplayNames.Add(MakeShared<FString>(DisplayName));
+			DisplayNameToFactoryName.Add(DisplayName, Factory->GetName());
+		}
+	}
+
+	// If the current Name property is empty and there are options, select the first one.
+	if (NameProperty.IsValid() && MuxerDisplayNames.Num() > 0)
+	{
+		FName CurrentName;
+		NameProperty->GetValue(CurrentName);
+		if (CurrentName.IsNone())
+		{
+			NameProperty->SetValue(DisplayNameToFactoryName[*MuxerDisplayNames[0]]);
+		}
+	}
+}
+
+void FTmvMediaMuxerSettingsCustomization::OnMuxerSelectionChanged(TSharedPtr<FString> InNewSelection, ESelectInfo::Type InSelectInfo)
+{
+	if (InNewSelection.IsValid() && NameProperty.IsValid())
+	{
+		if (const FName* FactoryName = DisplayNameToFactoryName.Find(*InNewSelection))
+		{
+			NameProperty->SetValue(*FactoryName);
+		}
+	}
+}
+
+FText FTmvMediaMuxerSettingsCustomization::GetSelectedMuxerDisplayName() const
+{
+	if (!NameProperty.IsValid())
+	{
+		return FText::GetEmpty();
+	}
+
+	FName CurrentName;
+	NameProperty->GetValue(CurrentName);
+
+	// Find the display name for the current factory name.
+	for (const auto& [DisplayName, FactoryName] : DisplayNameToFactoryName)
+	{
+		if (FactoryName == CurrentName)
+		{
+			return FText::FromString(DisplayName);
+		}
+	}
+
+	// Fallback: show the raw name if not found in registered factories.
+	return FText::FromName(CurrentName);
+}
+
+#undef LOCTEXT_NAMESPACE
diff --git a/Engine/Plugins/Media/TmvMedia/Source/TmvMediaEditor/Private/Customizations/TmvMediaMuxerSettingsCustomization.h b/Engine/Plugins/Media/TmvMedia/Source/TmvMediaEditor/Private/Customizations/TmvMediaMuxerSettingsCustomization.h
new file mode 100644
--- /dev/null
+++ b/Engine/Plugins/Media/TmvMedia/Source/TmvMediaEditor/Private/Customizations/TmvMediaMuxerSettingsCustomization.h
@@ -0,0 +1,44 @@
+// Copyright Epic Games, Inc. All Rights Reserved.
+
+#pragma once
+
+#include "IPropertyTypeCustomization.h"
+
+class IPropertyHandle;
+
+/**
+ * Implements a property type customization for FTmvMediaTranscodeMuxerSettings.
+ * The Name field is displayed as a combo box populated from registered ITmvMediaMuxerFactory instances.
+ */
+class FTmvMediaMuxerSettingsCustomization : public IPropertyTypeCustomization
+{
+public:
+	/** Make an instance of this property type customization. */
+	static TSharedRef<IPropertyTypeCustomization> MakeInstance();
+
+	//~ Begin IPropertyTypeCustomization
+	virtual void CustomizeHeader(TSharedRef<IPropertyHandle> InPropertyHandle, FDetailWidgetRow& InHeaderRow,
+		IPropertyTypeCustomizationUtils& InCustomizationUtils) override;
+	virtual void CustomizeChildren(TSharedRef<IPropertyHandle> InPropertyHandle, IDetailChildrenBuilder& InChildBuilder,
+		IPropertyTypeCustomizationUtils& InCustomizationUtils) override;
+	//~ End IPropertyTypeCustomization
+
+private:
+	/** Refresh the list of available muxer factories. */
+	void RefreshMuxerOptions();
+
+	/** Called when the combo box selection changes. */
+	void OnMuxerSelectionChanged(TSharedPtr<FString> InNewSelection, ESelectInfo::Type InSelectInfo);
+
+	/** Returns the currently selected muxer display name. */
+	FText GetSelectedMuxerDisplayName() const;
+
+	/** Handle to the Name property of FTmvMediaTranscodeMuxerSettings. */
+	TSharedPtr<IPropertyHandle> NameProperty;
+
+	/** Available muxer option strings for the combo box (display names). */
+	TArray<TSharedPtr<FString>> MuxerDisplayNames;
+
+	/** Mapping from display name to factory FName. */
+	TMap<FString, FName> DisplayNameToFactoryName;
+};
diff --git a/Engine/Plugins/Media/TmvMedia/Source/TmvMediaEditor/Private/Widgets/STmvMediaTranscoderJobControls.cpp b/Engine/Plugins/Media/TmvMedia/Source/TmvMediaEditor/Private/Widgets/STmvMediaTranscoderJobControls.cpp
--- a/Engine/Plugins/Media/TmvMedia/Source/TmvMediaEditor/Private/Widgets/STmvMediaTranscoderJobControls.cpp
+++ b/Engine/Plugins/Media/TmvMedia/Source/TmvMediaEditor/Private/Widgets/STmvMediaTranscoderJobControls.cpp
@@ -19,6 +19,8 @@
 #include "Transcoder/TmvMediaTmvFrameEncoder.h"
 #include "Transcoder/TmvMediaTranscodeJob.h"
 #include "Transcoder/TmvMediaTranscodeList.h"
+#include "Transcoder/TmvMediaContainerTranscodeMuxer.h"
+#include "Transcoder/TmvMediaFileSequenceTranscodeMuxer.h"
 #include "Transcoder/TmvMediaTranscodeMuxer.h"
 #include "UObject/UObjectIterator.h"
 #include "Widgets/Input/SButton.h"
@@ -99,7 +101,14 @@ namespace UE::TmvMediaEditor::Transcode
 		TranscodeJob->AddStage(NewObject<UTmvMediaFrameConverter>(TranscodeJob, NAME_None));
 		TranscodeJob->AddStage(NewObject<UTmvMediaTmvFrameEncoder>(TranscodeJob, NAME_None));
 		//TranscodeJob->AddStage(NewObject<UTmvMediaTestFrameEncoder>(TranscodeJob, NAME_None));
-		TranscodeJob->AddStage(NewObject<UTmvMediaTranscodeMuxer>(TranscodeJob, NAME_None));
+		if (TranscodeJob->Settings.OutputFormat == ETmvMediaTranscodeOutputFormat::Container)
+		{
+			TranscodeJob->AddStage(NewObject<UTmvMediaContainerTranscodeMuxer>(TranscodeJob, NAME_None));
+		}
+		else
+		{
+			TranscodeJob->AddStage(NewObject<UTmvMediaFileSequenceTranscodeMuxer>(TranscodeJob, NAME_None));
+		}
 		TranscodeJob->GetStage<UTmvMediaFrameEncoder>()->SetEncoderOptions(InJobItem->EncoderOptions);
 
 		// Extension editor tasks (can only be implemented in editor module).
diff --git a/Engine/Plugins/Media/TmvMedia/Source/TmvMediaEditor/Private/Widgets/STmvMediaTranscoderJobDetails.cpp b/Engine/Plugins/Media/TmvMedia/Source/TmvMediaEditor/Private/Widgets/STmvMediaTranscoderJobDetails.cpp
--- a/Engine/Plugins/Media/TmvMedia/Source/TmvMediaEditor/Private/Widgets/STmvMediaTranscoderJobDetails.cpp
+++ b/Engine/Plugins/Media/TmvMedia/Source/TmvMediaEditor/Private/Widgets/STmvMediaTranscoderJobDetails.cpp
@@ -4,6 +4,7 @@
 
 #include "Customizations/TmvMediaDirectoryPathCustomization.h"
 #include "Customizations/TmvMediaFilePathCustomization.h"
+#include "Customizations/TmvMediaMuxerSettingsCustomization.h"
 #include "DetailLayoutBuilder.h"
 #include "Editor.h"
 #include "Encoder/TmvMediaEncoderOptions.h"
@@ -267,6 +268,8 @@ void STmvMediaTranscoderJobDetails::RefreshJobSettingsDetailsView()
 				FOnGetPropertyTypeCustomizationInstance::CreateStatic(FTmvMediaFilePathCustomization::MakeInstance));
 			JobDetailsView->GetDetailsView()->RegisterInstancedCustomPropertyTypeLayout(FName(TEXT("DirectoryPath")),
 				FOnGetPropertyTypeCustomizationInstance::CreateStatic(FTmvMediaDirectoryPathCustomization::MakeInstance));
+			JobDetailsView->GetDetailsView()->RegisterInstancedCustomPropertyTypeLayout(FName(TEXT("TmvMediaTranscodeMuxerSettings")),
+				FOnGetPropertyTypeCustomizationInstance::CreateStatic(FTmvMediaMuxerSettingsCustomization::MakeInstance));
 
 			JobDetailsView->GetDetailsView()->OnFinishedChangingProperties().AddSP(this, &STmvMediaTranscoderJobDetails::OnFinishedChangingProperties);
 		}
diff --git a/Engine/Plugins/Media/TmvMedia/Source/TmvMediaEditor/Private/Widgets/TmvMediaTranscodeTaskMakeMediaSource.cpp b/Engine/Plugins/Media/TmvMedia/Source/TmvMediaEditor/Private/Widgets/TmvMediaTranscodeTaskMakeMediaSource.cpp
--- a/Engine/Plugins/Media/TmvMedia/Source/TmvMediaEditor/Private/Widgets/TmvMediaTranscodeTaskMakeMediaSource.cpp
+++ b/Engine/Plugins/Media/TmvMedia/Source/TmvMediaEditor/Private/Widgets/TmvMediaTranscodeTaskMakeMediaSource.cpp
@@ -13,6 +13,7 @@
 #include "Transcoder/TmvMediaFrameEncoder.h"
 #include "Transcoder/TmvMediaFrameProducer.h"
 #include "Transcoder/TmvMediaTranscodeJob.h"
+#include "Transcoder/TmvMediaTranscodeMuxer.h"
 #include "UObject/Package.h"
 #include "UObject/SavePackage.h"
 #include "Utils/TmvMediaPathUtils.h"
@@ -236,7 +237,19 @@ namespace UE::TmvMediaEditor::TranscodeTask
 			// Using the job name as the media source name.
 			const FString FullMediaSourcePath = MakeMediaSourceFullAssetName(InTranscodeJob);
 			const FString MediaSourceName = FPackageName::GetShortName(FullMediaSourcePath);
-			const FString FullMediaPath = InTranscodeJob->Settings.GetAbsoluteOutputPath();
+			const bool bIsContainerOutput = InTranscodeJob->Settings.OutputFormat == ETmvMediaTranscodeOutputFormat::Container;
+
+			// For containers, the media path is the output file (e.g. OutputPath/BaseName.tmv).
+			// For file sequences, the media path is the output directory.
+			FString FullMediaPath;
+			if (bIsContainerOutput)
+			{
+				FullMediaPath = UTmvMediaTranscodeMuxer::GetContainerOutputFilePath(InTranscodeJob);
+			}
+			else
+			{
+				FullMediaPath = InTranscodeJob->Settings.GetAbsoluteOutputPath();
+			}
 
 			UPackage* MediaSourcePackage = CreatePackage(*FullMediaSourcePath);
 
@@ -288,9 +301,15 @@ namespace UE::TmvMediaEditor::TranscodeTask
 			// Create new media source.
 			if (!MediaSource)
 			{
-				// Retrieve the produced file extension, if any.
+				// Retrieve the produced file extension.
+				// For containers, use the container extension (e.g. "tmv").
+				// For file sequences, use the per-frame codec extension (e.g. "apv1").
 				FString FileExtension;
-				if (const UTmvMediaFrameEncoder* EncoderStage = InTranscodeJob->GetStage<UTmvMediaFrameEncoder>())
+				if (bIsContainerOutput)
+				{
+					FileExtension = FPaths::GetExtension(FullMediaPath);
+				}
+				else if (const UTmvMediaFrameEncoder* EncoderStage = InTranscodeJob->GetStage<UTmvMediaFrameEncoder>())
 				{
 					if (const FTmvMediaEncoderOptions* EncoderOptions = EncoderStage->GetEncoderOptions())
 					{
diff --git a/Engine/Plugins/Media/TmvMedia/Source/TmvMediaMp4Utils/Private/TmvMediaMp4Demuxer.cpp b/Engine/Plugins/Media/TmvMedia/Source/TmvMediaMp4Utils/Private/TmvMediaMp4Demuxer.cpp
new file mode 100644
--- /dev/null
+++ b/Engine/Plugins/Media/TmvMedia/Source/TmvMediaMp4Utils/Private/TmvMediaMp4Demuxer.cpp
@@ -0,0 +1,462 @@
+// Copyright Epic Games, Inc. All Rights Reserved.
+
+#include "TmvMediaMp4Demuxer.h"
+#include "TmvMediaMp4UtilsLog.h"
+#include "MP4Boxes.h"
+
+namespace UE::TmvMedia
+{
+
+FTmvMediaMp4Demuxer::~FTmvMediaMp4Demuxer()
+{
+	Close();
+}
+
+ETmvMediaContainerResult FTmvMediaMp4Demuxer::OpenFile(const FString& InFilePath)
+{
+	Close();
+
+	TSharedPtr<MP4Utilities::FMP4FileDataReader, ESPMode::ThreadSafe> FileReader = MP4Utilities::FMP4FileDataReader::Create();
+	if (!FileReader.IsValid())
+	{
+		LastError = TEXT("Failed to create MP4 file data reader.");
+		return ETmvMediaContainerResult::Fail;
+	}
+
+	if (!FileReader->Open(InFilePath))
+	{
+		LastError = FString::Printf(TEXT("Failed to open file: %s"), *InFilePath);
+		return ETmvMediaContainerResult::Fail;
+	}
+
+	DataReader = FileReader;
+
+	// Locate and read root boxes from the file.
+	MP4Utilities::FMP4BoxLocator Locator;
+	// Read these root-level boxes: ftyp, moov, moof, sidx. Stop scanning after mdat in standard files.
+	const TArray<uint32> ReadBoxes = {
+		MP4Utilities::MakeBoxAtom('f','t','y','p'),
+		MP4Utilities::MakeBoxAtom('m','o','o','v'),
+		MP4Utilities::MakeBoxAtom('m','o','o','f'),
+		MP4Utilities::MakeBoxAtom('s','i','d','x'),
+	};
+
+	if (!Locator.LocateAndReadRootBoxes(
+		RootBoxes,
+		DataReader,
+		TArray<uint32>(),   // InFirstBoxes (no order constraint)
+		TArray<uint32>(),   // InStopAfterBoxes
+		false,              // bStopWithMDAT
+		ReadBoxes,          // InReadDataOfBoxes
+		MP4Utilities::IMP4DataReaderBase::FCancellationCheckDelegate()))
+	{
+		LastError = FString::Printf(TEXT("Failed to locate root boxes: %s"), *Locator.GetLastError());
+		return ETmvMediaContainerResult::Fail;
+	}
+
+	return ParseContainer();
+}
+
+ETmvMediaContainerResult FTmvMediaMp4Demuxer::OpenBuffer(TConstArrayView<uint8> InData)
+{
+	Close();
+
+	if (!MP4Utilities::FMP4BoxLocator::LocateAndReadRootBoxesFromBuffer(RootBoxes, InData))
+	{
+		LastError = TEXT("Failed to parse root boxes from buffer.");
+		return ETmvMediaContainerResult::Fail;
+	}
+
+	// Create a buffer data reader for sample data access.
+	DataReader = MP4Utilities::FMP4BufferDataReader::Create(InData);
+
+	return ParseContainer();
+}
+
+ETmvMediaContainerResult FTmvMediaMp4Demuxer::ParseContainer()
+{
+	// Find the moov root box.
+	TSharedPtr<MP4Utilities::FMP4BoxData, ESPMode::ThreadSafe> MoovBoxData;
+	TArray<TSharedPtr<MP4Utilities::FMP4BoxData, ESPMode::ThreadSafe>> MoofBoxDatas;
+
+	for (const TSharedPtr<MP4Utilities::FMP4BoxData, ESPMode::ThreadSafe>& Box : RootBoxes)
+	{
+		if (Box->Type == MP4Utilities::MakeBoxAtom('m','o','o','v'))
+		{
+			MoovBoxData = Box;
+		}
+		else if (Box->Type == MP4Utilities::MakeBoxAtom('m','o','o','f'))
+		{
+			MoofBoxDatas.Add(Box);
+		}
+	}
+
+	if (!MoovBoxData.IsValid())
+	{
+		LastError = TEXT("No 'moov' box found in container.");
+		return ETmvMediaContainerResult::Fail;
+	}
+
+	// Parse the moov box tree.
+	MP4Boxes::FMP4BoxTreeParser MoovParser;
+	if (!MoovParser.ParseBoxTree(MoovBoxData))
+	{
+		LastError = TEXT("Failed to parse 'moov' box tree.");
+		return ETmvMediaContainerResult::Fail;
+	}
+
+	TSharedPtr<MP4Boxes::FMP4BoxBase, ESPMode::ThreadSafe> MoovTree = MoovParser.GetBoxTree();
+	if (!MoovTree.IsValid())
+	{
+		LastError = TEXT("Empty 'moov' box tree.");
+		return ETmvMediaContainerResult::Fail;
+	}
+
+	// Get movie header for overall duration.
+	TSharedPtr<MP4Boxes::FMP4BoxMVHD, ESPMode::ThreadSafe> MvhdBox =
+		MoovTree->FindBoxRecursive<MP4Boxes::FMP4BoxMVHD>(MP4Utilities::MakeBoxAtom('m','v','h','d'));
+	if (MvhdBox.IsValid())
+	{
+		MovieDuration = MvhdBox->GetDuration();
+	}
+
+	// Parse moof boxes for fragmented files.
+	TArray<TSharedPtr<MP4Boxes::FMP4BoxBase>> AllRootBoxes;
+	AllRootBoxes.Add(MoovTree);
+	for (const TSharedPtr<MP4Utilities::FMP4BoxData, ESPMode::ThreadSafe>& MoofData : MoofBoxDatas)
+	{
+		MP4Boxes::FMP4BoxTreeParser MoofParser;
+		if (MoofParser.ParseBoxTree(MoofData))
+		{
+			AllRootBoxes.Add(MoofParser.GetBoxTree());
+		}
+	}
+
+	// Find mvex box for fragmented files.
+	TSharedPtr<MP4Boxes::FMP4BoxMVEX, ESPMode::ThreadSafe> MvexBox =
+		MoovTree->FindBoxRecursive<MP4Boxes::FMP4BoxMVEX>(MP4Utilities::MakeBoxAtom('m','v','e','x'));
+
+	// Find all trak boxes within moov.
+	TArray<TSharedPtr<MP4Boxes::FMP4BoxTRAK, ESPMode::ThreadSafe>> TrakBoxes;
+	MoovTree->GetAllBoxInstances<MP4Boxes::FMP4BoxTRAK>(TrakBoxes, MP4Utilities::MakeBoxAtom('t','r','a','k'));
+
+	int32 TrackIndex = 0;
+	for (const TSharedPtr<MP4Boxes::FMP4BoxTRAK, ESPMode::ThreadSafe>& TrakBox : TrakBoxes)
+	{
+		// Build fragment info only for fragmented files (mvex + moof present).
+		TSharedPtr<MP4Boxes::FMP4Track::FFragmentInfo, ESPMode::ThreadSafe> FragmentInfo;
+		if (MvexBox.IsValid())
+		{
+			FragmentInfo = MakeShared<MP4Boxes::FMP4Track::FFragmentInfo, ESPMode::ThreadSafe>();
+			FragmentInfo->Prepare(TrakBox, MvexBox, AllRootBoxes);
+		}
+
+		// Create the track.
+		TSharedPtr<MP4Boxes::FMP4Track, ESPMode::ThreadSafe> Track =
+			MP4Boxes::FMP4Track::Create(TrakBox, FragmentInfo);
+
+		if (!Track->Prepare(MovieDuration, MovieDuration))
+		{
+			UE_LOGF(LogTmvMediaMp4Utils, Warning, "FTmvMediaMp4Demuxer: Failed to prepare track %d: %ls", TrackIndex, *Track->GetLastError());
+			++TrackIndex;
+			continue;
+		}
+
+		FTrackState State;
+		State.Track = Track;
+		State.Iterator = Track->CreateIterator();
+
+		// Build cached track info.
+		FTmvMediaDemuxerTrackInfo& Info = State.CachedInfo;
+		Info.TrackIndex = TrackStates.Num();
+
+		// Determine track type from handler and STSD.
+		TSharedPtr<MP4Boxes::FMP4BoxHDLR, ESPMode::ThreadSafe> HdlrBox =
+			TrakBox->FindBoxRecursive<MP4Boxes::FMP4BoxHDLR>(MP4Utilities::MakeBoxAtom('h','d','l','r'));
+		if (HdlrBox.IsValid())
+		{
+			uint32 HandlerType = HdlrBox->GetHandlerType();
+			if (HandlerType == MP4Utilities::MakeBoxAtom('v','i','d','e'))
+			{
+				Info.TrackType = ETmvMediaTrackType::Video;
+			}
+			else if (HandlerType == MP4Utilities::MakeBoxAtom('s','o','u','n'))
+			{
+				Info.TrackType = ETmvMediaTrackType::Audio;
+			}
+		}
+
+		// Get timescale and duration from mdhd.
+		TSharedPtr<MP4Boxes::FMP4BoxMDHD, ESPMode::ThreadSafe> MdhdBox =
+			TrakBox->FindBoxRecursive<MP4Boxes::FMP4BoxMDHD>(MP4Utilities::MakeBoxAtom('m','d','h','d'));
+		if (MdhdBox.IsValid())
+		{
+			Info.Timescale = MdhdBox->GetTimescale();
+			MP4Utilities::FFractionalTime Duration = MdhdBox->GetDuration();
+			Info.Duration = Duration.GetNumerator();
+			Info.LanguageCode = MdhdBox->GetLanguageCode639_2T();
+		}
+
+		Info.NumSamples = Track->GetNumberOfSamples();
+
+		// Get constant sample duration from stts if all samples share the same delta.
+		TSharedPtr<MP4Boxes::FMP4BoxSTTS, ESPMode::ThreadSafe> SttsBox =
+			TrakBox->FindBoxRecursive<MP4Boxes::FMP4BoxSTTS>(MP4Utilities::MakeBoxAtom('s','t','t','s'));
+		if (SttsBox.IsValid())
+		{
+			const auto& Entries = SttsBox->GetEntries();
+			if (Entries.Num() == 1)
+			{
+				Info.ConstantSampleDuration = Entries[0].sample_delta;
+			}
+		}
+
+		// Get sample entry format and type-specific info from stsd.
+		TSharedPtr<MP4Boxes::FMP4BoxSTSD, ESPMode::ThreadSafe> StsdBox =
+			TrakBox->FindBoxRecursive<MP4Boxes::FMP4BoxSTSD>(MP4Utilities::MakeBoxAtom('s','t','s','d'));
+		if (StsdBox.IsValid())
+		{
+			// GetSampleType() must be called first — it triggers parsing of STSD children
+			// into typed sample entry objects (FMP4BoxVisualSampleEntry, etc.).
+			MP4Boxes::FMP4BoxSampleEntry::ESampleType SampleType = StsdBox->GetSampleType();
+
+			if (SampleType == MP4Boxes::FMP4BoxSampleEntry::ESampleType::Video)
+			{
+				TArray<TSharedPtr<MP4Boxes::FMP4BoxVisualSampleEntry, ESPMode::ThreadSafe>> VisualEntries;
+				StsdBox->GetSampleDescriptions(VisualEntries);
+				if (VisualEntries.Num() > 0)
+				{
+					Info.SampleEntryFormat = VisualEntries[0]->GetType();
+					Info.DisplayWidth = VisualEntries[0]->GetWidth();
+					Info.DisplayHeight = VisualEntries[0]->GetHeight();
+
+					for (const TSharedPtr<MP4Boxes::FMP4BoxBase, ESPMode::ThreadSafe>& ChildBox : VisualEntries[0]->GetChildren())
+					{
+						Info.CodecSpecificBoxes.Add(ChildBox->GetBoxDataRAW());
+					}
+				}
+			}
+			else if (SampleType == MP4Boxes::FMP4BoxSampleEntry::ESampleType::Audio)
+			{
+				TArray<TSharedPtr<MP4Boxes::FMP4BoxAudioSampleEntry, ESPMode::ThreadSafe>> AudioEntries;
+				StsdBox->GetSampleDescriptions(AudioEntries);
+				if (AudioEntries.Num() > 0)
+				{
+					Info.SampleEntryFormat = AudioEntries[0]->GetType();
+					Info.SamplingRate = AudioEntries[0]->GetSampleRate();
+					Info.NumberOfChannels = static_cast<uint16>(AudioEntries[0]->GetChannelCount());
+
+					for (const TSharedPtr<MP4Boxes::FMP4BoxBase, ESPMode::ThreadSafe>& ChildBox : AudioEntries[0]->GetChildren())
+					{
+						Info.CodecSpecificBoxes.Add(ChildBox->GetBoxDataRAW());
+					}
+				}
+			}
+		}
+
+		// Get track name from metadata.
+		const MP4Boxes::FMP4TrackMetadataCommon& Metadata = Track->GetCommonMetadata();
+		Info.TrackName = Metadata.Name;
+		if (!Metadata.LanguageTag.IsEmpty())
+		{
+			Info.LanguageCode = Metadata.LanguageTag;
+		}
+
+		// Get dimensions from tkhd as fallback for video tracks.
+		if (Info.TrackType == ETmvMediaTrackType::Video && (Info.DisplayWidth == 0 || Info.DisplayHeight == 0))
+		{
+			TSharedPtr<MP4Boxes::FMP4BoxTKHD, ESPMode::ThreadSafe> TkhdBox =
+				TrakBox->FindBoxRecursive<MP4Boxes::FMP4BoxTKHD>(MP4Utilities::MakeBoxAtom('t','k','h','d'));
+			if (TkhdBox.IsValid())
+			{
+				Info.DisplayWidth = TkhdBox->GetWidth();
+				Info.DisplayHeight = TkhdBox->GetHeight();
+			}
+		}
+
+		TrackStates.Add(MoveTemp(State));
+		++TrackIndex;
+	}
+
+	if (TrackStates.IsEmpty())
+	{
+		LastError = TEXT("No valid tracks found in container.");
+		return ETmvMediaContainerResult::Fail;
+	}
+
+	return ETmvMediaContainerResult::Success;
+}
+
+int32 FTmvMediaMp4Demuxer::GetTrackCount() const
+{
+	return TrackStates.Num();
+}
+
+ETmvMediaContainerResult FTmvMediaMp4Demuxer::GetTrackInfo(int32 InTrackIndex, FTmvMediaDemuxerTrackInfo& OutTrackInfo) const
+{
+	if (InTrackIndex < 0 || InTrackIndex >= TrackStates.Num())
+	{
+		return ETmvMediaContainerResult::Fail;
+	}
+
+	OutTrackInfo = TrackStates[InTrackIndex].CachedInfo;
+	return ETmvMediaContainerResult::Success;
+}
+
+ETmvMediaContainerResult FTmvMediaMp4Demuxer::ReadSample(int32 InTrackIndex, FTmvMediaDemuxerSample& OutSample)
+{
+	if (InTrackIndex < 0 || InTrackIndex >= TrackStates.Num())
+	{
+		return ETmvMediaContainerResult::Fail;
+	}
+
+	FTrackState& State = TrackStates[InTrackIndex];
+	if (!State.Iterator.IsValid() || !State.Iterator->IsValid())
+	{
+		return ETmvMediaContainerResult::EndOfStream;
+	}
+
+	ETmvMediaContainerResult Result = PopulateSampleFromIterator(State.Iterator, OutSample, /*bReadData=*/ true);
+
+	// Always advance the iterator regardless of success or failure.
+	// On failure, skipping the problematic sample prevents the caller from getting
+	// stuck in an infinite retry loop on the same broken sample.
+	if (!State.Iterator->Next())
+	{
+		// Iterator exhausted, next call will return EndOfStream.
+	}
+
+	return Result;
+}
+
+ETmvMediaContainerResult FTmvMediaMp4Demuxer::ReadSampleInfo(int32 InTrackIndex, FTmvMediaDemuxerSample& OutSample)
+{
+	if (InTrackIndex < 0 || InTrackIndex >= TrackStates.Num())
+	{
+		return ETmvMediaContainerResult::Fail;
+	}
+
+	FTrackState& State = TrackStates[InTrackIndex];
+	if (!State.Iterator.IsValid() || !State.Iterator->IsValid())
+	{
+		return ETmvMediaContainerResult::EndOfStream;
+	}
+
+	ETmvMediaContainerResult Result = PopulateSampleFromIterator(State.Iterator, OutSample, /*bReadData=*/ false);
+
+	// Always advance the iterator regardless of success or failure.
+	if (!State.Iterator->Next())
+	{
+		// Iterator exhausted, next call will return EndOfStream.
+	}
+
+	return Result;
+}
+
+ETmvMediaContainerResult FTmvMediaMp4Demuxer::PopulateSampleFromIterator(
+	const TSharedPtr<MP4Boxes::FMP4Track::FIterator, ESPMode::ThreadSafe>& InIterator,
+	FTmvMediaDemuxerSample& OutSample,
+	bool bReadData)
+{
+	OutSample.SampleNumber = InIterator->GetSampleNumber();
+	OutSample.DTS = InIterator->GetDTS().GetNumerator();
+	OutSample.PTS = InIterator->GetPTS().GetNumerator();
+	OutSample.Duration = static_cast<uint32>(InIterator->GetDuration().GetNumerator());
+	OutSample.bIsKeyframe = InIterator->IsSyncOrRAPSample();
+	OutSample.FileOffset = InIterator->GetSampleFileOffset();
+	OutSample.SampleSize = InIterator->GetSampleSize();
+
+	if (bReadData && DataReader.IsValid() && OutSample.SampleSize > 0)
+	{
+		OutSample.Data.SetNumUninitialized(OutSample.SampleSize);
+		int64 BytesRead = DataReader->ReadData(
+			OutSample.Data.GetData(),
+			OutSample.SampleSize,
+			OutSample.FileOffset,
+			MP4Utilities::IMP4DataReaderBase::FCancellationCheckDelegate());
+
+		if (BytesRead < 0)
+		{
+			LastError = FString::Printf(TEXT("Failed to read sample data at offset %lld, size %lld: %s"),
+				OutSample.FileOffset, OutSample.SampleSize, *DataReader->GetLastError());
+			OutSample.Data.Empty();
+			return ETmvMediaContainerResult::Fail;
+		}
+
+		if (BytesRead < OutSample.SampleSize)
+		{
+			LastError = FString::Printf(TEXT("Partial read at offset %lld: expected %lld bytes, got %lld"),
+				OutSample.FileOffset, OutSample.SampleSize, BytesRead);
+			OutSample.Data.Empty();
+			return ETmvMediaContainerResult::Fail;
+		}
+	}
+	else
+	{
+		OutSample.Data.Empty();
+	}
+
+	return ETmvMediaContainerResult::Success;
+}
+
+ETmvMediaContainerResult FTmvMediaMp4Demuxer::Seek(int32 InTrackIndex, FTimespan InTime, FTimespan InLaterTimeThreshold)
+{
+	if (InTrackIndex < 0 || InTrackIndex >= TrackStates.Num())
+	{
+		return ETmvMediaContainerResult::Fail;
+	}
+
+	FTrackState& State = TrackStates[InTrackIndex];
+	if (!State.Track.IsValid())
+	{
+		return ETmvMediaContainerResult::Fail;
+	}
+
+	State.Iterator = State.Track->CreateIteratorAtKeyframe(InTime, InLaterTimeThreshold);
+	if (!State.Iterator.IsValid() || !State.Iterator->IsValid())
+	{
+		LastError = TEXT("Failed to seek to keyframe.");
+		return ETmvMediaContainerResult::Fail;
+	}
+
+	return ETmvMediaContainerResult::Success;
+}
+
+ETmvMediaContainerResult FTmvMediaMp4Demuxer::SeekToSample(int32 InTrackIndex, uint32 InSampleNumber)
+{
+	if (InTrackIndex < 0 || InTrackIndex >= TrackStates.Num())
+	{
+		return ETmvMediaContainerResult::Fail;
+	}
+
+	FTrackState& State = TrackStates[InTrackIndex];
+	if (!State.Track.IsValid())
+	{
+		return ETmvMediaContainerResult::Fail;
+	}
+
+	State.Iterator = State.Track->CreateIterator(InSampleNumber);
+	if (!State.Iterator.IsValid() || !State.Iterator->IsValid())
+	{
+		LastError = FString::Printf(TEXT("Failed to seek to sample %u."), InSampleNumber);
+		return ETmvMediaContainerResult::Fail;
+	}
+
+	return ETmvMediaContainerResult::Success;
+}
+
+void FTmvMediaMp4Demuxer::Close()
+{
+	TrackStates.Empty();
+	RootBoxes.Empty();
+	DataReader.Reset();
+	MovieDuration = MP4Utilities::FFractionalTime();
+	LastError.Empty();
+}
+
+FString FTmvMediaMp4Demuxer::GetLastError() const
+{
+	return LastError;
+}
+
+} // namespace UE::TmvMedia
diff --git a/Engine/Plugins/Media/TmvMedia/Source/TmvMediaMp4Utils/Private/TmvMediaMp4Demuxer.h b/Engine/Plugins/Media/TmvMedia/Source/TmvMediaMp4Utils/Private/TmvMediaMp4Demuxer.h
new file mode 100644
--- /dev/null
+++ b/Engine/Plugins/Media/TmvMedia/Source/TmvMediaMp4Utils/Private/TmvMediaMp4Demuxer.h
@@ -0,0 +1,70 @@
+// Copyright Epic Games, Inc. All Rights Reserved.
+
+#pragma once
+
+#include "Decoder/ITmvMediaDemuxer.h"
+#include "MP4BoxBase.h"
+#include "MP4Track.h"
+#include "MP4DataReader.h"
+#include "MP4Utilities.h"
+
+namespace UE::TmvMedia
+{
+
+/**
+ * MP4 container demuxer implementation.
+ * Wraps MP4Boxes/MP4Utilities for box tree parsing and sample iteration.
+ */
+class FTmvMediaMp4Demuxer : public ITmvMediaDemuxer
+{
+public:
+	FTmvMediaMp4Demuxer() = default;
+	virtual ~FTmvMediaMp4Demuxer() override;
+
+	// ITmvMediaDemuxer interface
+	virtual ETmvMediaContainerResult OpenFile(const FString& InFilePath) override;
+	virtual ETmvMediaContainerResult OpenBuffer(TConstArrayView<uint8> InData) override;
+	virtual int32 GetTrackCount() const override;
+	virtual ETmvMediaContainerResult GetTrackInfo(int32 InTrackIndex, FTmvMediaDemuxerTrackInfo& OutTrackInfo) const override;
+	virtual ETmvMediaContainerResult ReadSample(int32 InTrackIndex, FTmvMediaDemuxerSample& OutSample) override;
+	virtual ETmvMediaContainerResult ReadSampleInfo(int32 InTrackIndex, FTmvMediaDemuxerSample& OutSample) override;
+	virtual ETmvMediaContainerResult Seek(int32 InTrackIndex, FTimespan InTime, FTimespan InLaterTimeThreshold) override;
+	virtual ETmvMediaContainerResult SeekToSample(int32 InTrackIndex, uint32 InSampleNumber) override;
+	virtual void Close() override;
+	virtual FString GetLastError() const override;
+
+private:
+	/** Parse the box tree from loaded root boxes and build tracks. */
+	ETmvMediaContainerResult ParseContainer();
+
+	/** Populate an FTmvMediaDemuxerSample from the current iterator position. Optionally reads sample data. */
+	ETmvMediaContainerResult PopulateSampleFromIterator(
+		const TSharedPtr<MP4Boxes::FMP4Track::FIterator, ESPMode::ThreadSafe>& InIterator,
+		FTmvMediaDemuxerSample& OutSample,
+		bool bReadData);
+
+	/** Per-track state. */
+	struct FTrackState
+	{
+		TSharedPtr<MP4Boxes::FMP4Track, ESPMode::ThreadSafe> Track;
+		TSharedPtr<MP4Boxes::FMP4Track::FIterator, ESPMode::ThreadSafe> Iterator;
+		FTmvMediaDemuxerTrackInfo CachedInfo;
+	};
+
+	/** Data reader for file or buffer access. */
+	TSharedPtr<MP4Utilities::IMP4DataReaderBase, ESPMode::ThreadSafe> DataReader;
+
+	/** Parsed root boxes. */
+	TArray<TSharedPtr<MP4Utilities::FMP4BoxData, ESPMode::ThreadSafe>> RootBoxes;
+
+	/** Track states. */
+	TArray<FTrackState> TrackStates;
+
+	/** Movie duration from mvhd. */
+	MP4Utilities::FFractionalTime MovieDuration;
+
+	/** Last error message. */
+	FString LastError;
+};
+
+} // namespace UE::TmvMedia
diff --git a/Engine/Plugins/Media/TmvMedia/Source/TmvMediaMp4Utils/Private/TmvMediaMp4DemuxerFactory.cpp b/Engine/Plugins/Media/TmvMedia/Source/TmvMediaMp4Utils/Private/TmvMediaMp4DemuxerFactory.cpp
new file mode 100644
--- /dev/null
+++ b/Engine/Plugins/Media/TmvMedia/Source/TmvMediaMp4Utils/Private/TmvMediaMp4DemuxerFactory.cpp
@@ -0,0 +1,27 @@
+// Copyright Epic Games, Inc. All Rights Reserved.
+
+#include "TmvMediaMp4DemuxerFactory.h"
+
+#include "Internationalization/Internationalization.h"
+#include "TmvMediaMp4Demuxer.h"
+
+FName FTmvMediaMp4DemuxerFactory::GetName() const
+{
+	static const FName Name(TEXT("Tmv"));
+	return Name;
+}
+
+FText FTmvMediaMp4DemuxerFactory::GetDisplayName() const
+{
+	return NSLOCTEXT("TmvMediaMp4Utils", "DemuxerDisplayName", "Tiled-Mipmap Video");
+}
+
+TArray<FString> FTmvMediaMp4DemuxerFactory::GetSupportedContainerFormats() const
+{
+	return { TEXT("tmv") };
+}
+
+TSharedPtr<ITmvMediaDemuxer, ESPMode::ThreadSafe> FTmvMediaMp4DemuxerFactory::CreateDemuxer()
+{
+	return MakeShared<UE::TmvMedia::FTmvMediaMp4Demuxer, ESPMode::ThreadSafe>();
+}
diff --git a/Engine/Plugins/Media/TmvMedia/Source/TmvMediaMp4Utils/Private/TmvMediaMp4DemuxerFactory.h b/Engine/Plugins/Media/TmvMedia/Source/TmvMediaMp4Utils/Private/TmvMediaMp4DemuxerFactory.h
new file mode 100644
--- /dev/null
+++ b/Engine/Plugins/Media/TmvMedia/Source/TmvMediaMp4Utils/Private/TmvMediaMp4DemuxerFactory.h
@@ -0,0 +1,17 @@
+// Copyright Epic Games, Inc. All Rights Reserved.
+
+#pragma once
+
+#include "Decoder/ITmvMediaDemuxerFactory.h"
+
+class FTmvMediaMp4DemuxerFactory : public ITmvMediaDemuxerFactory
+{
+public:
+	virtual ~FTmvMediaMp4DemuxerFactory() override = default;
+
+	// ITmvMediaDemuxerFactory interface
+	virtual FName GetName() const override;
+	virtual FText GetDisplayName() const override;
+	virtual TArray<FString> GetSupportedContainerFormats() const override;
+	virtual TSharedPtr<ITmvMediaDemuxer, ESPMode::ThreadSafe> CreateDemuxer() override;
+};
diff --git a/Engine/Plugins/Media/TmvMedia/Source/TmvMediaMp4Utils/Private/TmvMediaMp4Muxer.cpp b/Engine/Plugins/Media/TmvMedia/Source/TmvMediaMp4Utils/Private/TmvMediaMp4Muxer.cpp
new file mode 100644
--- /dev/null
+++ b/Engine/Plugins/Media/TmvMedia/Source/TmvMediaMp4Utils/Private/TmvMediaMp4Muxer.cpp
@@ -0,0 +1,219 @@
+// Copyright Epic Games, Inc. All Rights Reserved.
+
+#include "TmvMediaMp4Muxer.h"
+#include "TmvMediaMp4UtilsLog.h"
+
+namespace UE::TmvMedia
+{
+
+FTmvMediaMp4Muxer::~FTmvMediaMp4Muxer()
+{
+	if (Mp4Muxer.IsValid())
+	{
+		Mp4Muxer->StopAndClose();
+		Mp4Muxer.Reset();
+	}
+}
+
+ETmvMediaContainerResult FTmvMediaMp4Muxer::Configure(const FTmvMediaMuxerConfig& InConfig)
+{
+	Mp4Muxer = IMP4RawMuxer::Create();
+	if (!Mp4Muxer.IsValid())
+	{
+		UE_LOGF(LogTmvMediaMp4Utils, Error, "FTmvMediaMp4Muxer::Configure Failed to create Mp4RawMuxer");
+		return ETmvMediaContainerResult::Fail;
+	}
+
+	IMP4RawMuxer::FConfiguration Mp4Config;
+	Mp4Config.OutputFilename = InConfig.OutputFilename;
+	Mp4Config.TemporaryFilename = InConfig.TemporaryFilename;
+	Mp4Config.InterleaveDuration = InConfig.InterleaveDuration;
+
+	switch (InConfig.OutputMode)
+	{
+	case FTmvMediaMuxerConfig::EOutputMode::Standard:
+		Mp4Config.MuxMode = IMP4RawMuxer::FConfiguration::EMuxMode::Standard;
+		break;
+	case FTmvMediaMuxerConfig::EOutputMode::WebOptimized:
+		Mp4Config.MuxMode = IMP4RawMuxer::FConfiguration::EMuxMode::WebOptimized;
+		break;
+	case FTmvMediaMuxerConfig::EOutputMode::Fragmented:
+		Mp4Config.MuxMode = IMP4RawMuxer::FConfiguration::EMuxMode::Fragmented;
+		break;
+	default:
+		Mp4Config.MuxMode = IMP4RawMuxer::FConfiguration::EMuxMode::Standard;
+		break;
+	}
+
+	if (!Mp4Muxer->Configure(Mp4Config))
+	{
+		LastError = Mp4Muxer->GetLastError();
+		UE_LOGF(LogTmvMediaMp4Utils, Error, "FTmvMediaMp4Muxer::Configure failed: %ls", *LastError);
+		return ETmvMediaContainerResult::Fail;
+	}
+
+	return ETmvMediaContainerResult::Success;
+}
+
+int32 FTmvMediaMp4Muxer::AddTrack(const FTmvMediaMuxerTrackConfig& InTrackConfig)
+{
+	if (!Mp4Muxer.IsValid())
+	{
+		LastError = TEXT("Muxer not configured. Call Configure() first.");
+		return -1;
+	}
+
+	IMP4RawMuxer::FTrackSpec TrackSpec;
+
+	// Set track type properties.
+	switch (InTrackConfig.TrackType)
+	{
+	case ETmvMediaTrackType::Video:
+	{
+		IMP4RawMuxer::FTrackSpec::FVideo Video;
+		Video.DisplayWidth = InTrackConfig.DisplayWidth;
+		Video.DisplayHeight = InTrackConfig.DisplayHeight;
+		Video.CompressorName = InTrackConfig.CompressorName;
+		TrackSpec.Properties.Emplace<IMP4RawMuxer::FTrackSpec::FVideo>(MoveTemp(Video));
+		break;
+	}
+	case ETmvMediaTrackType::Audio:
+	{
+		IMP4RawMuxer::FTrackSpec::FAudio Audio;
+		Audio.SamplingRate = InTrackConfig.SamplingRate;
+		Audio.NumberOfChannels = InTrackConfig.NumberOfChannels;
+		TrackSpec.Properties.Emplace<IMP4RawMuxer::FTrackSpec::FAudio>(MoveTemp(Audio));
+		break;
+	}
+	case ETmvMediaTrackType::Unspecified:
+		LastError = TEXT("\"Unspecified\" track type");
+		return -1;
+	}
+
+	TrackSpec.SampleEntryFormat = InTrackConfig.SampleEntryFormat;
+	TrackSpec.bIsAllKeyframes = InTrackConfig.bIsAllKeyframes;
+	TrackSpec.Timescale = InTrackConfig.Timescale;
+	TrackSpec.ConstantSampleDuration = InTrackConfig.ConstantSampleDuration;
+	TrackSpec.SampleEntryBoxes = InTrackConfig.CodecSpecificBoxes;
+	TrackSpec.SubSampleFlags = InTrackConfig.SubSampleFlags;
+	TrackSpec.Name = InTrackConfig.TrackName;
+
+	int32 TrackIndex = Mp4Muxer->AddTrack(TrackSpec);
+	if (TrackIndex < 0)
+	{
+		LastError = Mp4Muxer->GetLastError();
+		UE_LOGF(LogTmvMediaMp4Utils, Error, "FTmvMediaMp4Muxer::AddTrack failed: %ls", *LastError);
+	}
+
+	return TrackIndex;
+}
+
+ETmvMediaContainerResult FTmvMediaMp4Muxer::Start(FSampleRequestDelegate InSampleRequestDelegate, FStatusDelegate InStatusDelegate)
+{
+	if (!Mp4Muxer.IsValid())
+	{
+		LastError = TEXT("Muxer not configured. Call Configure() first.");
+		return ETmvMediaContainerResult::Fail;
+	}
+
+	// Adapt the MP4 muxer's delegate to the TmvMedia interface delegate.
+	IMP4RawMuxer::FRequestTrackDataDelegate Mp4RequestDelegate;
+	Mp4RequestDelegate.BindLambda([InSampleRequestDelegate](const TArray<IMP4RawMuxer::FSampleRequest>& InMp4Requests)
+	{
+		TArray<ITmvMediaMuxer::FSampleRequest> TmvRequests;
+		TmvRequests.Reserve(InMp4Requests.Num());
+		for (const IMP4RawMuxer::FSampleRequest& Mp4Req : InMp4Requests)
+		{
+			ITmvMediaMuxer::FSampleRequest TmvReq;
+			TmvReq.TrackIndex = Mp4Req.TrackIndex;
+			TmvReq.SampleNumber = Mp4Req.SampleNumber;
+			TmvRequests.Add(TmvReq);
+		}
+		InSampleRequestDelegate.ExecuteIfBound(TmvRequests);
+	});
+
+	IMP4RawMuxer::FStatusDelegate Mp4StatusDelegate;
+	Mp4StatusDelegate.BindLambda([InStatusDelegate](IMP4RawMuxer::EStatus InMp4Status)
+	{
+		ITmvMediaMuxer::EStatus TmvStatus = (InMp4Status == IMP4RawMuxer::EStatus::Finished)
+			? ITmvMediaMuxer::EStatus::Finished
+			: ITmvMediaMuxer::EStatus::Failed;
+		InStatusDelegate.ExecuteIfBound(TmvStatus);
+	});
+
+	if (!Mp4Muxer->Start(Mp4RequestDelegate, Mp4StatusDelegate))
+	{
+		LastError = Mp4Muxer->GetLastError();
+		UE_LOGF(LogTmvMediaMp4Utils, Error, "FTmvMediaMp4Muxer::Start failed: %ls", *LastError);
+		return ETmvMediaContainerResult::Fail;
+	}
+
+	return ETmvMediaContainerResult::Success;
+}
+
+ETmvMediaContainerResult FTmvMediaMp4Muxer::AddSample(int32 InTrackIndex, const FTmvMediaMuxerSample& InSample)
+{
+	if (!Mp4Muxer.IsValid())
+	{
+		LastError = TEXT("Muxer not configured.");
+		return ETmvMediaContainerResult::Fail;
+	}
+
+	IMP4RawMuxer::FTrackSample Mp4Sample;
+	Mp4Sample.Data = InSample.Data;
+	Mp4Sample.Duration = InSample.Duration;
+	Mp4Sample.DTS = InSample.DTS;
+	Mp4Sample.PTS = InSample.PTS;
+	Mp4Sample.SampleNumber = InSample.SampleNumber;
+	Mp4Sample.bIsKeyframe = InSample.bIsKeyframe;
+	Mp4Sample.bIsFinalSample = InSample.bIsFinalSample;
+
+	// Convert subsample info.
+	Mp4Sample.SubSamples.Reserve(InSample.SubSamples.Num());
+	for (const FTmvMediaMuxerSample::FSubSampleInfo& SubSample : InSample.SubSamples)
+	{
+		IMP4RawMuxer::FTrackSample::FSubSampleInfo Mp4Sub;
+		Mp4Sub.codec_specific_parameters = SubSample.CodecSpecificParameters;
+		Mp4Sub.subsample_size = SubSample.SubSampleSize;
+		Mp4Sub.subsample_priority = SubSample.SubSamplePriority;
+		Mp4Sub.discardable = SubSample.Discardable;
+		Mp4Sample.SubSamples.Add(Mp4Sub);
+	}
+
+	if (!Mp4Muxer->AddTrackSample(InTrackIndex, Mp4Sample))
+	{
+		LastError = Mp4Muxer->GetLastError();
+		UE_LOGF(LogTmvMediaMp4Utils, Error, "FTmvMediaMp4Muxer::AddSample failed for track %d, sample %u: %ls", InTrackIndex, InSample.SampleNumber, *LastError);
+		return ETmvMediaContainerResult::Fail;
+	}
+
+	return ETmvMediaContainerResult::Success;
+}
+
+ETmvMediaContainerResult FTmvMediaMp4Muxer::Finalize()
+{
+	if (!Mp4Muxer.IsValid())
+	{
+		LastError = TEXT("Muxer not configured.");
+		return ETmvMediaContainerResult::Fail;
+	}
+	
+	ETmvMediaContainerResult Result = ETmvMediaContainerResult::Success;
+
+	if (!Mp4Muxer->StopAndClose())
+	{
+		LastError = Mp4Muxer->GetLastError();
+		UE_LOGF(LogTmvMediaMp4Utils, Error, "FTmvMediaMp4Muxer::Finalize failed: %ls", *LastError);
+		Result = ETmvMediaContainerResult::Fail;
+	}
+
+	Mp4Muxer.Reset();
+	return Result;
+}
+
+FString FTmvMediaMp4Muxer::GetLastError() const
+{
+	return LastError;
+}
+
+} // namespace UE::TmvMedia
diff --git a/Engine/Plugins/Media/TmvMedia/Source/TmvMediaMp4Utils/Private/TmvMediaMp4Muxer.h b/Engine/Plugins/Media/TmvMedia/Source/TmvMediaMp4Utils/Private/TmvMediaMp4Muxer.h
new file mode 100644
--- /dev/null
+++ b/Engine/Plugins/Media/TmvMedia/Source/TmvMediaMp4Utils/Private/TmvMediaMp4Muxer.h
@@ -0,0 +1,37 @@
+// Copyright Epic Games, Inc. All Rights Reserved.
+
+#pragma once
+
+#include "Encoder/ITmvMediaMuxer.h"
+#include "MP4Muxer.h"
+
+namespace UE::TmvMedia
+{
+
+/**
+ * MP4 container muxer implementation.
+ * Wraps IMP4RawMuxer from the MP4Muxer module.
+ */
+class FTmvMediaMp4Muxer : public ITmvMediaMuxer
+{
+public:
+	FTmvMediaMp4Muxer() = default;
+	virtual ~FTmvMediaMp4Muxer() override;
+
+	// ITmvMediaMuxer interface
+	virtual ETmvMediaContainerResult Configure(const FTmvMediaMuxerConfig& InConfig) override;
+	virtual int32 AddTrack(const FTmvMediaMuxerTrackConfig& InTrackConfig) override;
+	virtual ETmvMediaContainerResult Start(FSampleRequestDelegate InSampleRequestDelegate, FStatusDelegate InStatusDelegate) override;
+	virtual ETmvMediaContainerResult AddSample(int32 InTrackIndex, const FTmvMediaMuxerSample& InSample) override;
+	virtual ETmvMediaContainerResult Finalize() override;
+	virtual FString GetLastError() const override;
+
+private:
+	/** The underlying MP4 raw muxer. */
+	TSharedPtr<IMP4RawMuxer, ESPMode::ThreadSafe> Mp4Muxer;
+
+	/** Last error message. */
+	FString LastError;
+};
+
+} // namespace UE::TmvMedia
diff --git a/Engine/Plugins/Media/TmvMedia/Source/TmvMediaMp4Utils/Private/TmvMediaMp4MuxerFactory.cpp b/Engine/Plugins/Media/TmvMedia/Source/TmvMediaMp4Utils/Private/TmvMediaMp4MuxerFactory.cpp
new file mode 100644
--- /dev/null
+++ b/Engine/Plugins/Media/TmvMedia/Source/TmvMediaMp4Utils/Private/TmvMediaMp4MuxerFactory.cpp
@@ -0,0 +1,32 @@
+// Copyright Epic Games, Inc. All Rights Reserved.
+
+#include "TmvMediaMp4MuxerFactory.h"
+
+#include "Internationalization/Internationalization.h"
+#include "TmvMediaMp4Muxer.h"
+
+FName FTmvMediaMp4MuxerFactory::GetName() const
+{
+	static const FName Name(TEXT("Tmv"));
+	return Name;
+}
+
+FText FTmvMediaMp4MuxerFactory::GetDisplayName() const
+{
+	return NSLOCTEXT("TmvMediaMp4Utils", "MuxerDisplayName", "Tiled-Mipmap Video");
+}
+
+TArray<FString> FTmvMediaMp4MuxerFactory::GetSupportedContainerFormats() const
+{
+	return { TEXT("tmv") };
+}
+
+FString FTmvMediaMp4MuxerFactory::GetFileExtension() const
+{
+	return TEXT("tmv");
+}
+
+TSharedPtr<ITmvMediaMuxer, ESPMode::ThreadSafe> FTmvMediaMp4MuxerFactory::CreateMuxer()
+{
+	return MakeShared<UE::TmvMedia::FTmvMediaMp4Muxer, ESPMode::ThreadSafe>();
+}
diff --git a/Engine/Plugins/Media/TmvMedia/Source/TmvMediaMp4Utils/Private/TmvMediaMp4MuxerFactory.h b/Engine/Plugins/Media/TmvMedia/Source/TmvMediaMp4Utils/Private/TmvMediaMp4MuxerFactory.h
new file mode 100644
--- /dev/null
+++ b/Engine/Plugins/Media/TmvMedia/Source/TmvMediaMp4Utils/Private/TmvMediaMp4MuxerFactory.h
@@ -0,0 +1,18 @@
+// Copyright Epic Games, Inc. All Rights Reserved.
+
+#pragma once
+
+#include "Encoder/ITmvMediaMuxerFactory.h"
+
+class FTmvMediaMp4MuxerFactory : public ITmvMediaMuxerFactory
+{
+public:
+	virtual ~FTmvMediaMp4MuxerFactory() override = default;
+
+	// ITmvMediaMuxerFactory interface
+	virtual FName GetName() const override;
+	virtual FText GetDisplayName() const override;
+	virtual TArray<FString> GetSupportedContainerFormats() const override;
+	virtual FString GetFileExtension() const override;
+	virtual TSharedPtr<ITmvMediaMuxer, ESPMode::ThreadSafe> CreateMuxer() override;
+};
diff --git a/Engine/Plugins/Media/TmvMedia/Source/TmvMediaMp4Utils/Private/TmvMediaMp4UtilsLog.h b/Engine/Plugins/Media/TmvMedia/Source/TmvMediaMp4Utils/Private/TmvMediaMp4UtilsLog.h
new file mode 100644
--- /dev/null
+++ b/Engine/Plugins/Media/TmvMedia/Source/TmvMediaMp4Utils/Private/TmvMediaMp4UtilsLog.h
@@ -0,0 +1,7 @@
+// Copyright Epic Games, Inc. All Rights Reserved.
+
+#pragma once
+
+#include "Logging/LogMacros.h"
+
+DECLARE_LOG_CATEGORY_EXTERN(LogTmvMediaMp4Utils, Log, All);
diff --git a/Engine/Plugins/Media/TmvMedia/Source/TmvMediaMp4Utils/Private/TmvMediaMp4UtilsModule.cpp b/Engine/Plugins/Media/TmvMedia/Source/TmvMediaMp4Utils/Private/TmvMediaMp4UtilsModule.cpp
new file mode 100644
--- /dev/null
+++ b/Engine/Plugins/Media/TmvMedia/Source/TmvMediaMp4Utils/Private/TmvMediaMp4UtilsModule.cpp
@@ -0,0 +1,43 @@
+// Copyright Epic Games, Inc. All Rights Reserved.
+
+#include "TmvMediaMp4UtilsLog.h"
+#include "TmvMediaMp4MuxerFactory.h"
+#include "TmvMediaMp4DemuxerFactory.h"
+#include "ITmvMediaModule.h"
+#include "Modules/ModuleInterface.h"
+#include "Modules/ModuleManager.h"
+
+DEFINE_LOG_CATEGORY(LogTmvMediaMp4Utils);
+
+/** Implementation of TmvMediaMp4Utils Module. */
+class FTmvMediaMp4UtilsModule : public IModuleInterface
+{
+public:
+	//~ Begin IModuleInterface
+	virtual void StartupModule() override
+	{
+		MuxerFactory = MakeShared<FTmvMediaMp4MuxerFactory>();
+		ITmvMediaModule::GetOrLoad().RegisterMuxerFactory(MuxerFactory);
+
+		DemuxerFactory = MakeShared<FTmvMediaMp4DemuxerFactory>();
+		ITmvMediaModule::GetOrLoad().RegisterDemuxerFactory(DemuxerFactory);
+	}
+
+	virtual void ShutdownModule() override
+	{
+		if (ITmvMediaModule* TmvMediaModule = ITmvMediaModule::Get())
+		{
+			TmvMediaModule->UnregisterMuxerFactory(MuxerFactory);
+			TmvMediaModule->UnregisterDemuxerFactory(DemuxerFactory);
+		}
+		MuxerFactory.Reset();
+		DemuxerFactory.Reset();
+	}
+	//~ End IModuleInterface
+
+private:
+	TSharedPtr<FTmvMediaMp4MuxerFactory, ESPMode::ThreadSafe> MuxerFactory;
+	TSharedPtr<FTmvMediaMp4DemuxerFactory, ESPMode::ThreadSafe> DemuxerFactory;
+};
+
+IMPLEMENT_MODULE(FTmvMediaMp4UtilsModule, TmvMediaMp4Utils);
diff --git a/Engine/Plugins/Media/TmvMedia/Source/TmvMediaMp4Utils/TmvMediaMp4Utils.Build.cs b/Engine/Plugins/Media/TmvMedia/Source/TmvMediaMp4Utils/TmvMediaMp4Utils.Build.cs
new file mode 100644
--- /dev/null
+++ b/Engine/Plugins/Media/TmvMedia/Source/TmvMediaMp4Utils/TmvMediaMp4Utils.Build.cs
@@ -0,0 +1,20 @@
+// Copyright Epic Games, Inc. All Rights Reserved.
+
+namespace UnrealBuildTool.Rules
+{
+	public class TmvMediaMp4Utils : ModuleRules
+	{
+		public TmvMediaMp4Utils(ReadOnlyTargetRules Target) : base(Target)
+		{
+			PrivateDependencyModuleNames.AddRange(
+				new string[]
+				{
+					"Core",
+					"TmvMedia",
+					"MP4Muxer",
+					"MP4Boxes",
+					"MP4Utilities",
+				});
+		}
+	}
+}
diff --git a/Engine/Plugins/Media/TmvMedia/TmvMedia.uplugin b/Engine/Plugins/Media/TmvMedia/TmvMedia.uplugin
--- a/Engine/Plugins/Media/TmvMedia/TmvMedia.uplugin
+++ b/Engine/Plugins/Media/TmvMedia/TmvMedia.uplugin
@@ -37,6 +37,19 @@
 			"LoadingPhase" : "PostEngineInit",
 			"PlatformAllowList" : [ "Win64" ],
 			"PlatformArchitectureDenyList": [ "Win64:arm64" ]
+		},
+		{
+			"Name" : "TmvMediaMp4Utils",
+			"Type" : "Runtime",
+			"LoadingPhase" : "PostEngineInit",
+			"PlatformAllowList" : [ "Win64" ]
+		}
+	],
+	"Plugins" :
+	[
+		{
+			"Name" : "MP4Utilities",
+			"Enabled" : true
 		}
 	]
 }
